HoloLensAR now disabled by default
-Avoids pulling in MixedRealityInterop.dll when we don't want to use it.  It is problematic because that dll is dependent on windows 10.

diff --git a/Engine/Plugins/Runtime/AR/Microsoft/HoloLensAR/HoloLensAR.uplugin b/Engine/Plugins/Runtime/AR/Microsoft/HoloLensAR/HoloLensAR.uplugin
--- a/Engine/Plugins/Runtime/AR/Microsoft/HoloLensAR/HoloLensAR.uplugin
+++ b/Engine/Plugins/Runtime/AR/Microsoft/HoloLensAR/HoloLensAR.uplugin
@@ -10,7 +10,7 @@
 	"DocsURL": "",
 	"MarketplaceURL": "",
 	"SupportURL": "",
-	"EnabledByDefault": true,
+	"EnabledByDefault": false,
 	"CanContainContent": true,
 	"IsBetaVersion": true,
 	"Installed": false,
