Merging using //UE4/Partner-Microsoft-XR/_to_//UE4/Release-4.23/

Skip partially installed SDK roots in HoloLens project configuration

diff --git a/Engine/Plugins/Runtime/AR/Microsoft/HoloLensAR/Source/HoloLensPlatformEditor/Private/HoloLensTargetSettingsCustomization.cpp b/Engine/Plugins/Runtime/AR/Microsoft/HoloLensAR/Source/HoloLensPlatformEditor/Private/HoloLensTargetSettingsCustomization.cpp
--- a/Engine/Plugins/Runtime/AR/Microsoft/HoloLensAR/Source/HoloLensPlatformEditor/Private/HoloLensTargetSettingsCustomization.cpp
+++ b/Engine/Plugins/Runtime/AR/Microsoft/HoloLensAR/Source/HoloLensPlatformEditor/Private/HoloLensTargetSettingsCustomization.cpp
@@ -47,26 +47,36 @@ TArray<FHoloLensSDKVersion> FHoloLensSDKVersion::GetSDKVersions()
 		for (DWORD dwRootIdx = 0; dwRootIdx < _countof(hRoots); ++dwRootIdx)
 		{
 			HKEY hKey;
-			if (RegOpenKeyExW(hRoots[dwRootIdx], L"Software\\Microsoft\\Windows Kits\\Installed Roots", 0, KEY_ENUMERATE_SUB_KEYS, &hKey) == ERROR_SUCCESS)
+			if (RegOpenKeyExW(hRoots[dwRootIdx], L"Software\\Microsoft\\Windows Kits\\Installed Roots", 0, KEY_ENUMERATE_SUB_KEYS | KEY_QUERY_VALUE, &hKey) == ERROR_SUCCESS)
 			{
 				DWORD dwIndex = 0;
 				WCHAR sSubKeyName[256];
 				DWORD dwSubKeyNameLen = _countof(sSubKeyName);
+				FString SDKPath;
+
+				if (ERROR_SUCCESS == RegQueryValueExW(hKey, L"KitsRoot10", 0, NULL, (LPBYTE)sSubKeyName, &dwSubKeyNameLen))
+				{
+					SDKPath = FString(dwSubKeyNameLen, sSubKeyName);
+				}
+				dwSubKeyNameLen = _countof(sSubKeyName);
 
 				while (RegEnumKeyExW(hKey, dwIndex, sSubKeyName, &dwSubKeyNameLen, NULL, NULL, NULL, NULL) == ERROR_SUCCESS)
 				{
 					int version[4];
 					if (swscanf_s(sSubKeyName, L"%d.%d.%d.%d", &version[0], &version[1], &version[2], &version[3]) == 4)
 					{
-						// HoloLens min version is 17763 for remoting and 18317 for hand tracking
-						// @todo JoeG - update this version when the official os/sdk are released
-						if (version[0] >= 10 && version[2] >= 17763)
-						{
-							new(SDKVersions) FHoloLensSDKVersion(FString(dwSubKeyNameLen, sSubKeyName), version[0], version[1], version[2], version[3]);
-						}
-						if (version[2] == 18362)
+						if (FPaths::FileExists(FPaths::Combine(SDKPath, TEXT("Include"), FString(dwSubKeyNameLen, sSubKeyName), TEXT("um"), TEXT("Windows.h"))))
 						{
-							bFound18362 = true;
+							// HoloLens min version is 17763 for remoting and 18317 for hand tracking
+							// @todo JoeG - update this version when the official os/sdk are released
+							if (version[0] >= 10 && version[2] >= 17763)
+							{
+								new(SDKVersions) FHoloLensSDKVersion(FString(dwSubKeyNameLen, sSubKeyName), version[0], version[1], version[2], version[3]);
+							}
+							if (version[2] == 18362)
+							{
+								bFound18362 = true;
+							}
 						}
 					}
 					dwSubKeyNameLen = _countof(sSubKeyName);
