UEVR-1470 Hololens d3d device loss
-Disabling passthrough camera due to instability.  This is a temporary workaround.

diff --git a/Engine/Plugins/Runtime/AR/Microsoft/HoloLensAR/Source/HoloLensAR/Private/HoloLensARSystem.cpp b/Engine/Plugins/Runtime/AR/Microsoft/HoloLensAR/Source/HoloLensAR/Private/HoloLensARSystem.cpp
--- a/Engine/Plugins/Runtime/AR/Microsoft/HoloLensAR/Source/HoloLensAR/Private/HoloLensARSystem.cpp
+++ b/Engine/Plugins/Runtime/AR/Microsoft/HoloLensAR/Source/HoloLensAR/Private/HoloLensARSystem.cpp
@@ -202,7 +202,7 @@ void FHoloLensARSystem::OnStartARSession(UARSessionConfig* InSessionConfig)
 void FHoloLensARSystem::SetupCameraImageSupport()
 {
 	// Remoting does not support CameraCapture currently.
-	if (WMRInterop->IsRemoting())
+	//if (WMRInterop->IsRemoting()) //TEMP Disabling passthrough camera, it has d3d corruption bugs.
 	{
 		return;
 	}
