Add explicit .winmd dependency for QR module for hololens
If Mixed Reality Interop fails to load due to windows version, don't exit if this is a commandlet (packaging)

diff --git a/Engine/Plugins/Runtime/AR/Microsoft/HoloLensAR/Source/HoloLensAR/HoloLensAR.Build.cs b/Engine/Plugins/Runtime/AR/Microsoft/HoloLensAR/Source/HoloLensAR/HoloLensAR.Build.cs
--- a/Engine/Plugins/Runtime/AR/Microsoft/HoloLensAR/Source/HoloLensAR/HoloLensAR.Build.cs
+++ b/Engine/Plugins/Runtime/AR/Microsoft/HoloLensAR/Source/HoloLensAR/HoloLensAR.Build.cs
@@ -35,11 +35,13 @@ public class HoloLensAR : ModuleRules
         {
             PublicDelayLoadDLLs.Add("QRCodesTrackerPlugin.dll");
             RuntimeDependencies.Add(System.IO.Path.Combine("$(EngineDir)/Binaries/ThirdParty/Windows/x64", "QRCodesTrackerPlugin.dll"));
+            RuntimeDependencies.Add(System.IO.Path.Combine("$(EngineDir)/Binaries/ThirdParty/Windows/x64", "QRCodesTrackerPlugin.winmd"));
         }
         if (Target.Platform == UnrealTargetPlatform.HoloLens)
         {
             PublicDelayLoadDLLs.Add("QRCodesTrackerPlugin.dll");
             RuntimeDependencies.Add(System.IO.Path.Combine("$(EngineDir)/Binaries/ThirdParty/HoloLens/ARM64", "QRCodesTrackerPlugin.dll"));
+            RuntimeDependencies.Add(System.IO.Path.Combine("$(EngineDir)/Binaries/ThirdParty/HoloLens/ARM64", "QRCodesTrackerPlugin.winmd"));
         }
 
         AddEngineThirdPartyPrivateStaticDependencies(Target, "WindowsMixedRealityInterop");
@@ -47,16 +49,18 @@ public class HoloLensAR : ModuleRules
 		PublicDefinitions.Add("WITH_WINDOWS_MIXED_REALITY=1");
 
 		// Add a dependency to SceneUnderstanding.dll if present
-		string SceneUnderstandingPath = System.IO.Path.Combine(Target.UEThirdPartyBinariesDirectory, "HoloLens", Target.WindowsPlatform.GetArchitectureSubpath(), "SceneUnderstanding.dll");
-		if (System.IO.File.Exists(SceneUnderstandingPath))
+		string SceneUnderstandingDllPath = System.IO.Path.Combine(Target.UEThirdPartyBinariesDirectory, "HoloLens", Target.WindowsPlatform.GetArchitectureSubpath(), "SceneUnderstanding.dll");
+        string SceneUnderstandingWinMDPath = System.IO.Path.Combine(Target.UEThirdPartyBinariesDirectory, "HoloLens", Target.WindowsPlatform.GetArchitectureSubpath(), "SceneUnderstanding.winmd");
+        if (System.IO.File.Exists(SceneUnderstandingDllPath) && System.IO.File.Exists(SceneUnderstandingWinMDPath))
 		{
 			PublicDelayLoadDLLs.Add("SceneUnderstanding.dll");
-			RuntimeDependencies.Add(SceneUnderstandingPath);
-			PublicDefinitions.Add("WITH_SCENE_UNDERSTANDING=1");
-		}
-		else
-		{
-			PublicDefinitions.Add("WITH_SCENE_UNDERSTANDING=0");
+			RuntimeDependencies.Add(SceneUnderstandingDllPath);
+            RuntimeDependencies.Add(SceneUnderstandingWinMDPath);
+            PublicDefinitions.Add("WITH_SCENE_UNDERSTANDING=1");
 		}
+        else
+        {
+            PublicDefinitions.Add("WITH_SCENE_UNDERSTANDING=0");
+        }
 	}
 }
diff --git a/Engine/Plugins/Runtime/WindowsMixedReality/Source/WindowsMixedRealityHMD/Private/WindowsMixedRealityHMD.cpp b/Engine/Plugins/Runtime/WindowsMixedReality/Source/WindowsMixedRealityHMD/Private/WindowsMixedRealityHMD.cpp
--- a/Engine/Plugins/Runtime/WindowsMixedReality/Source/WindowsMixedRealityHMD/Private/WindowsMixedRealityHMD.cpp
+++ b/Engine/Plugins/Runtime/WindowsMixedReality/Source/WindowsMixedRealityHMD/Private/WindowsMixedRealityHMD.cpp
@@ -191,7 +191,14 @@ namespace WindowsMixedReality
 					"Windows Mixed Reality is not supported on this Windows version. \nNote: UE4 only supports Windows Mixed Reality on Windows 10 Release {0} or higher. Current version: {1}")),
 					FText::FromString(FString::FromInt(MIN_WIN_10_VERSION_FOR_WMR)), FText::FromString(OSVersionLabel));
 				FMessageDialog::Open(EAppMsgType::Ok, ErrorText);
-				UE_LOG(LogWmrHmd, Error, TEXT("%s"), *ErrorText.ToString());
+				if (IsRunningCommandlet())
+				{
+					UE_LOG(LogWmrHmd, Warning, TEXT("%s"), *ErrorText.ToString());
+				}
+				else
+				{
+					UE_LOG(LogWmrHmd, Error, TEXT("%s"), *ErrorText.ToString());
+				}
 			}
 
 #else // !PLATFORM_HOLOLENS
