We can't make explicit assumptions on what SDK's a user has installed when updating the relevant ui.

diff --git a/Engine/Plugins/Runtime/AR/Microsoft/HoloLensAR/Source/HoloLensPlatformEditor/Private/HoloLensTargetSettingsCustomization.cpp b/Engine/Plugins/Runtime/AR/Microsoft/HoloLensAR/Source/HoloLensPlatformEditor/Private/HoloLensTargetSettingsCustomization.cpp
--- a/Engine/Plugins/Runtime/AR/Microsoft/HoloLensAR/Source/HoloLensPlatformEditor/Private/HoloLensTargetSettingsCustomization.cpp
+++ b/Engine/Plugins/Runtime/AR/Microsoft/HoloLensAR/Source/HoloLensPlatformEditor/Private/HoloLensTargetSettingsCustomization.cpp
@@ -755,17 +755,14 @@ void FHoloLensTargetSettingsCustomization::OnAutoDetectWin10SDKChanged(ECheckBox
 		WindowsSDKSelector->SetSelectedItem(MakeShared<FString>());
 		WindowsSDKSelector->SetEnabled(false);
 	}
-	else if (GetDefault<UHoloLensTargetSettings>()->CompilerVersion == ECompilerVersion::VisualStudio2017)
-	{
-		Win10SDKVersionPropertyHandle->SetValue(*PlatformVersionOptions.Last());
-		WindowsSDKSelector->SetSelectedItem(PlatformVersionOptions.Last());
-		WindowsSDKSelector->SetEnabled(true);
-	}
 	else
 	{
-		Win10SDKVersionPropertyHandle->SetValue(*PlatformVersionOptions[2]);
-		WindowsSDKSelector->SetSelectedItem(PlatformVersionOptions[2]);
-		WindowsSDKSelector->SetEnabled(true);
+		if (PlatformVersionOptions.Num() > 0)
+		{
+			Win10SDKVersionPropertyHandle->SetValue(*PlatformVersionOptions[0]);
+			WindowsSDKSelector->SetSelectedItem(PlatformVersionOptions[0]);
+			WindowsSDKSelector->SetEnabled(true);
+		}
 	}
 	Win10SDKVersionPropertyHandle->NotifyPostChange();
 }
