Fix for HL2 static lighting in packaged builds

diff --git a/Engine/Plugins/Runtime/AR/Microsoft/HoloLensAR/Source/HoloLensTargetPlatform/Private/HoloLensTargetPlatform.cpp b/Engine/Plugins/Runtime/AR/Microsoft/HoloLensAR/Source/HoloLensTargetPlatform/Private/HoloLensTargetPlatform.cpp
--- a/Engine/Plugins/Runtime/AR/Microsoft/HoloLensAR/Source/HoloLensTargetPlatform/Private/HoloLensTargetPlatform.cpp
+++ b/Engine/Plugins/Runtime/AR/Microsoft/HoloLensAR/Source/HoloLensTargetPlatform/Private/HoloLensTargetPlatform.cpp
@@ -79,12 +79,22 @@ ITargetDevicePtr FHoloLensTargetPlatform::GetDefaultDevice() const
 
 bool FHoloLensTargetPlatform::SupportsFeature(ETargetPlatformFeatures Feature) const
 {
-	if (Feature == ETargetPlatformFeatures::Packaging)
+	switch (Feature)
 	{
-		return true;
-	}
+		case ETargetPlatformFeatures::Packaging:
+			return true;
+
+		case ETargetPlatformFeatures::LowQualityLightmaps:
+			// HoloLens 2 is mobile renderer, thus only supports low quality light maps
+			return true;
 
-	return TTargetPlatformBase<FHoloLensPlatformProperties>::SupportsFeature(Feature);
+		case ETargetPlatformFeatures::HighQualityLightmaps:
+			// HoloLens 2 is mobile renderer, thus only supports low quality light maps
+			return false;
+
+		default:
+			return TTargetPlatformBase<FHoloLensPlatformProperties>::SupportsFeature(Feature);
+	}
 }
 
 #if WITH_ENGINE
diff --git a/Engine/Source/Runtime/Core/Public/HoloLens/HoloLensPlatformProperties.h b/Engine/Source/Runtime/Core/Public/HoloLens/HoloLensPlatformProperties.h
--- a/Engine/Source/Runtime/Core/Public/HoloLens/HoloLensPlatformProperties.h
+++ b/Engine/Source/Runtime/Core/Public/HoloLens/HoloLensPlatformProperties.h
@@ -51,6 +51,18 @@ struct FHoloLensPlatformProperties : public FGenericPlatformProperties
 	{
 		return true;
 	}
+
+	static FORCEINLINE bool SupportsLowQualityLightmaps()
+	{
+		// HoloLens 2 is mobile renderer, thus only supports low quality light maps
+		return true;
+	}
+
+	static FORCEINLINE bool SupportsHighQualityLightmaps()
+	{
+		// HoloLens 2 is mobile renderer, thus only supports low quality light maps
+		return false; 
+	}
 };
 
 #ifdef PROPERTY_HEADER_SHOULD_DEFINE_TYPE
diff --git a/Engine/Source/Runtime/Engine/Public/SceneManagement.h b/Engine/Source/Runtime/Engine/Public/SceneManagement.h
--- a/Engine/Source/Runtime/Engine/Public/SceneManagement.h
+++ b/Engine/Source/Runtime/Engine/Public/SceneManagement.h
@@ -306,7 +306,7 @@ static const int32 LQ_LIGHTMAP_COEF_INDEX = 2;
 
 /** Compile out low quality lightmaps to save memory */
 // @todo-mobile: Need to fix this!
-#define ALLOW_LQ_LIGHTMAPS (PLATFORM_DESKTOP || PLATFORM_IOS || PLATFORM_ANDROID || PLATFORM_HTML5 || PLATFORM_SWITCH || PLATFORM_LUMIN)
+#define ALLOW_LQ_LIGHTMAPS (PLATFORM_DESKTOP || PLATFORM_IOS || PLATFORM_ANDROID || PLATFORM_HTML5 || PLATFORM_SWITCH || PLATFORM_LUMIN || PLATFORM_HOLOLENS)
 
 /** Compile out high quality lightmaps to save memory */
 #define ALLOW_HQ_LIGHTMAPS 1
