UE-79885 HL2 CL integrations
Merging //UE4/Partner-Microsoft-XR to Release-4.23 (//UE4/Release-4.23)
8513053 - Update HL2 remoting to 2.0.8.0
8579684 - Transform HoloLens spatial mapping into Unreal world space

diff --git a/Engine/Plugins/Runtime/AR/Microsoft/HoloLensAR/Source/HoloLensAR/Private/HoloLensARSystem.cpp b/Engine/Plugins/Runtime/AR/Microsoft/HoloLensAR/Source/HoloLensAR/Private/HoloLensARSystem.cpp
--- a/Engine/Plugins/Runtime/AR/Microsoft/HoloLensAR/Source/HoloLensAR/Private/HoloLensARSystem.cpp
+++ b/Engine/Plugins/Runtime/AR/Microsoft/HoloLensAR/Source/HoloLensAR/Private/HoloLensARSystem.cpp
@@ -536,9 +536,9 @@ void FHoloLensARSystem::SetupMeshObserver()
 
 	// Get the settings for triangle density and mapping volume size
 	float TriangleDensity = 500.f;
-	GConfig->GetFloat(TEXT("/Script/HoloLensTargetPlatform.HoloLensTargetSettings"), TEXT("MaxTrianglesPerCubicMeter"), TriangleDensity, GEngineIni);
+	GConfig->GetFloat(TEXT("/Script/HoloLensPlatformEditor.HoloLensTargetSettings"), TEXT("MaxTrianglesPerCubicMeter"), TriangleDensity, GEngineIni);
 	float VolumeSize = 1.f;
-	GConfig->GetFloat(TEXT("/Script/HoloLensTargetPlatform.HoloLensTargetSettings"), TEXT("SpatialMeshingVolumeSize"), VolumeSize, GEngineIni);
+	GConfig->GetFloat(TEXT("/Script/HoloLensPlatformEditor.HoloLensTargetSettings"), TEXT("SpatialMeshingVolumeSize"), VolumeSize, GEngineIni);
 
 	WMRInterop->StartSpatialMapping(TriangleDensity, VolumeSize, &StartMeshUpdates_Raw, &AllocateMeshBuffers_Raw, &EndMeshUpdates_Raw);
 }
@@ -564,12 +564,16 @@ void FHoloLensARSystem::AllocateMeshBuffers(MeshUpdate* InMeshUpdate)
 		MeshUpdate->Indices.AddUninitialized(InMeshUpdate->NumIndices);
 		InMeshUpdate->Indices = MeshUpdate->Indices.GetData();
 
+		const FTransform TrackingToWorldTransform = TrackingSystem->GetTrackingToWorldTransform();
+		
 		// The transform information is only updated when the vertices are updated so it needs to be captured here
 		FVector Translation(InMeshUpdate->Translation[0], InMeshUpdate->Translation[1], InMeshUpdate->Translation[2]);
+		Translation = TrackingToWorldTransform.TransformPosition(Translation);
 		MeshUpdate->Location = Translation;
 		FVector Scale(InMeshUpdate->Scale[0], InMeshUpdate->Scale[1], InMeshUpdate->Scale[2]);
 		MeshUpdate->Scale = Scale;
 		FQuat Rotation(InMeshUpdate->Rotation[0], InMeshUpdate->Rotation[1], InMeshUpdate->Rotation[2], InMeshUpdate->Rotation[3]);
+		Rotation = TrackingToWorldTransform.TransformRotation(Rotation);
 		MeshUpdate->Rotation = Rotation;
 	}
 
