CIS fixes from AR integration

diff --git a/Engine/Plugins/Runtime/AR/Google/GoogleARCore/Source/GoogleARCoreBase/Private/GoogleARCoreOpenGLContext.cpp b/Engine/Plugins/Runtime/AR/Google/GoogleARCore/Source/GoogleARCoreBase/Private/GoogleARCoreOpenGLContext.cpp
--- a/Engine/Plugins/Runtime/AR/Google/GoogleARCore/Source/GoogleARCoreBase/Private/GoogleARCoreOpenGLContext.cpp
+++ b/Engine/Plugins/Runtime/AR/Google/GoogleARCore/Source/GoogleARCoreBase/Private/GoogleARCoreOpenGLContext.cpp
@@ -1,4 +1,4 @@
-// Copyright 2017 Google Inc.
+// Copyright 1998-2019 Epic Games, Inc. All Rights Reserved.
 
 #include "GoogleARCoreOpenGLContext.h"
 
diff --git a/Engine/Plugins/Runtime/AR/Google/GoogleARCore/Source/GoogleARCoreBase/Private/GoogleARCoreOpenGLContext.h b/Engine/Plugins/Runtime/AR/Google/GoogleARCore/Source/GoogleARCoreBase/Private/GoogleARCoreOpenGLContext.h
--- a/Engine/Plugins/Runtime/AR/Google/GoogleARCore/Source/GoogleARCoreBase/Private/GoogleARCoreOpenGLContext.h
+++ b/Engine/Plugins/Runtime/AR/Google/GoogleARCore/Source/GoogleARCoreBase/Private/GoogleARCoreOpenGLContext.h
@@ -1,4 +1,4 @@
-// Copyright 2017 Google Inc.
+// Copyright 1998-2019 Epic Games, Inc. All Rights Reserved.
 
 #pragma once
 
diff --git a/Engine/Plugins/Runtime/AR/Microsoft/HoloLensAR/Source/HoloLensAR/Private/HoloLensARSystem.cpp b/Engine/Plugins/Runtime/AR/Microsoft/HoloLensAR/Source/HoloLensAR/Private/HoloLensARSystem.cpp
--- a/Engine/Plugins/Runtime/AR/Microsoft/HoloLensAR/Source/HoloLensAR/Private/HoloLensARSystem.cpp
+++ b/Engine/Plugins/Runtime/AR/Microsoft/HoloLensAR/Source/HoloLensAR/Private/HoloLensARSystem.cpp
@@ -109,11 +109,21 @@ void FHoloLensARSystem::OnARSystemInitialized()
 	UE_LOG(LogHoloLensAR, Log, TEXT("HoloLens AR system has been initialized"));
 }
 
+bool FHoloLensARSystem::IsARAvailable() const override
+{
+	return true;
+}
+
 EARTrackingQuality FHoloLensARSystem::OnGetTrackingQuality() const
 {
 	return TrackingQuality;
 }
 
+EARTrackingQualityReason FHoloLensARSystem::OnGetTrackingQualityReason() const
+{
+	return EARTrackingQualityReason::None;
+}
+
 void OnTrackingChanged_Raw(WindowsMixedReality::HMDSpatialLocatability InTrackingState)
 {
 	FHoloLensARSystem* HoloLensARThis = FHoloLensModuleAR::GetHoloLensARSystem().Get();
diff --git a/Engine/Plugins/Runtime/AR/Microsoft/HoloLensAR/Source/HoloLensAR/Public/HoloLensARSystem.h b/Engine/Plugins/Runtime/AR/Microsoft/HoloLensAR/Source/HoloLensAR/Public/HoloLensARSystem.h
--- a/Engine/Plugins/Runtime/AR/Microsoft/HoloLensAR/Source/HoloLensAR/Public/HoloLensARSystem.h
+++ b/Engine/Plugins/Runtime/AR/Microsoft/HoloLensAR/Source/HoloLensAR/Public/HoloLensARSystem.h
@@ -75,8 +75,10 @@ private:
 	//~ FGCObject
 
 	//~IARSystemSupport
+	virtual bool IsARAvailable() const override;
 	virtual void OnARSystemInitialized() override;
 	virtual EARTrackingQuality OnGetTrackingQuality() const override;
+	virtual EARTrackingQualityReason OnGetTrackingQualityReason() const override;
 	virtual void OnStartARSession(UARSessionConfig* InSessionConfig) override;
 	virtual void OnPauseARSession() override;
 	virtual void OnStopARSession() override;
