CIS fix for HoloLens

diff --git a/Engine/Plugins/Runtime/AR/Microsoft/HoloLensAR/Source/HoloLensAR/Private/HoloLensARSystem.cpp b/Engine/Plugins/Runtime/AR/Microsoft/HoloLensAR/Source/HoloLensAR/Private/HoloLensARSystem.cpp
--- a/Engine/Plugins/Runtime/AR/Microsoft/HoloLensAR/Source/HoloLensAR/Private/HoloLensARSystem.cpp
+++ b/Engine/Plugins/Runtime/AR/Microsoft/HoloLensAR/Source/HoloLensAR/Private/HoloLensARSystem.cpp
@@ -109,7 +109,7 @@ void FHoloLensARSystem::OnARSystemInitialized()
 	UE_LOG(LogHoloLensAR, Log, TEXT("HoloLens AR system has been initialized"));
 }
 
-bool FHoloLensARSystem::IsARAvailable() const override
+bool FHoloLensARSystem::IsARAvailable() const
 {
 	return true;
 }
