UE-79885 HL2 CL integrations
-Fix for a crash Jackson saw at least once when closing vr preview while meshing was running.

diff --git a/Engine/Plugins/Runtime/AR/Microsoft/HoloLensAR/Source/HoloLensAR/Private/HoloLensARSystem.cpp b/Engine/Plugins/Runtime/AR/Microsoft/HoloLensAR/Source/HoloLensAR/Private/HoloLensARSystem.cpp
--- a/Engine/Plugins/Runtime/AR/Microsoft/HoloLensAR/Source/HoloLensAR/Private/HoloLensARSystem.cpp
+++ b/Engine/Plugins/Runtime/AR/Microsoft/HoloLensAR/Source/HoloLensAR/Private/HoloLensARSystem.cpp
@@ -697,6 +697,13 @@ void FHoloLensARSystem::AddOrUpdateMesh(FMeshUpdate* CurrentMesh)
 		{
 			// Attach this component to our single origin actor
 			AAROriginActor* OriginActor = AAROriginActor::GetOriginActor();
+
+			// During shutdown we can get a mesh update after the OriginActor has been destroyed, just return in that case.
+			if (OriginActor == nullptr)
+			{
+				return;
+			}
+
 			UMRMeshComponent* MRMesh = NewObject<UMRMeshComponent>(OriginActor);
 
 			// Set the occlusion and wireframe defaults
