Update from Hololens2-Preview @ CL 6623263

diff --git a/Engine/Plugins/Runtime/AR/Microsoft/HoloLensAR/HoloLensAR.uplugin b/Engine/Plugins/Runtime/AR/Microsoft/HoloLensAR/HoloLensAR.uplugin
--- a/Engine/Plugins/Runtime/AR/Microsoft/HoloLensAR/HoloLensAR.uplugin
+++ b/Engine/Plugins/Runtime/AR/Microsoft/HoloLensAR/HoloLensAR.uplugin
@@ -1,19 +1,23 @@
 {
-	"FileVersion" : 3,
-	"Version" : 1,
-	"VersionName" : "1.0",
-	"FriendlyName" : "HoloLens",
-	"Description" : "HoloLens Support",
-	"Category" : "Augmented Reality",
-	"CreatedBy" : "Epic Games, Inc. and Microsoft",
-	"CreatedByURL" : "http://epicgames.com",
-	"DocsURL" : "",
-	"MarketplaceURL" : "",
-	"SupportURL" : "",
-	"EnabledByDefault" : true,
-	"CanContainContent" : true,
-	"IsBetaVersion" : true,
-	"Installed" : false,
+	"FileVersion": 3,
+	"Version": 1,
+	"VersionName": "1.0",
+	"FriendlyName": "HoloLens",
+	"Description": "HoloLens Support",
+	"Category": "Augmented Reality",
+	"CreatedBy": "Epic Games, Inc. and Microsoft",
+	"CreatedByURL": "http://epicgames.com",
+	"DocsURL": "",
+	"MarketplaceURL": "",
+	"SupportURL": "",
+	"EnabledByDefault": true,
+	"CanContainContent": true,
+	"IsBetaVersion": true,
+	"Installed": false,
+	"SupportedTargetPlatforms": [
+		"HoloLens",
+		"Win64"
+	],
 	"Modules": [
 		{
 			"Name": "HoloLensAR",
