Fix for case of include in HoloLens DLC Settings customization for ToolTip

diff --git a/Engine/Plugins/Runtime/AR/Microsoft/HoloLensAR/Source/HoloLensPlatformEditor/Private/HoloLensDlcSettingsCustomization.cpp b/Engine/Plugins/Runtime/AR/Microsoft/HoloLensAR/Source/HoloLensPlatformEditor/Private/HoloLensDlcSettingsCustomization.cpp
--- a/Engine/Plugins/Runtime/AR/Microsoft/HoloLensAR/Source/HoloLensPlatformEditor/Private/HoloLensDlcSettingsCustomization.cpp
+++ b/Engine/Plugins/Runtime/AR/Microsoft/HoloLensAR/Source/HoloLensPlatformEditor/Private/HoloLensDlcSettingsCustomization.cpp
@@ -11,7 +11,7 @@
 #include "IDetailGroup.h"
 #include "Widgets/Input/STextComboBox.h"
 #include "Interfaces/IPluginManager.h"
-#include "Widgets/STooltip.h"
+#include "Widgets/SToolTip.h"
 #include "Settings/ProjectPackagingSettings.h"
 #include "HoloLensLocalizedResourcesCustomization.h"
 
