MERGE: Marked raw triangle mesh tracked geometry as "World" scene type

diff --git a/Engine/Plugins/Runtime/AR/Microsoft/HoloLensAR/Source/HoloLensAR/Private/HoloLensARSystem.cpp b/Engine/Plugins/Runtime/AR/Microsoft/HoloLensAR/Source/HoloLensAR/Private/HoloLensARSystem.cpp
--- a/Engine/Plugins/Runtime/AR/Microsoft/HoloLensAR/Source/HoloLensAR/Private/HoloLensARSystem.cpp
+++ b/Engine/Plugins/Runtime/AR/Microsoft/HoloLensAR/Source/HoloLensAR/Private/HoloLensARSystem.cpp
@@ -651,6 +651,8 @@ void FHoloLensARSystem::AddOrUpdateMesh(FMeshUpdate* CurrentMesh)
 			FTransform(CurrentMesh->Rotation, CurrentMesh->Location, CurrentMesh->Scale),
 			// @todo JoeG - add support for an alignment transform
 			FTransform::Identity);
+		// Mark this as a world mesh that isn't recognized as a particular scene type, since it is loose triangles
+		NewUpdatedGeometry->SetObjectClassification(EARObjectClassification::World);
 		if (NewUpdatedGeometry->GetUnderlyingMesh() == nullptr)
 		{
 			// Attach this component to our single origin actor
