Marked each of the HoloLens AR plugin modules as being used on Win64/HoloLens so the Mac doesn't try to load them

diff --git a/Engine/Plugins/Runtime/AR/Microsoft/HoloLensAR/HoloLensAR.uplugin b/Engine/Plugins/Runtime/AR/Microsoft/HoloLensAR/HoloLensAR.uplugin
--- a/Engine/Plugins/Runtime/AR/Microsoft/HoloLensAR/HoloLensAR.uplugin
+++ b/Engine/Plugins/Runtime/AR/Microsoft/HoloLensAR/HoloLensAR.uplugin
@@ -22,22 +22,38 @@
 		{
 			"Name": "HoloLensAR",
 			"Type": "Runtime",
-			"LoadingPhase": "PostConfigInit"
+			"LoadingPhase": "PostConfigInit",
+			"WhitelistPlatforms": [
+				"HoloLens",
+				"Win64"
+			]
 		},
 		{
 			"Name": "HoloLensPlatformEditor",
 			"Type": "Editor",
-			"LoadingPhase": "Default"
+			"LoadingPhase": "Default",
+			"WhitelistPlatforms": [
+				"HoloLens",
+				"Win64"
+			]
 		},
 		{
 			"Name": "HoloLensDeviceDetector",
 			"Type": "Editor",
-			"LoadingPhase": "Default"
+			"LoadingPhase": "Default",
+			"WhitelistPlatforms": [
+				"HoloLens",
+				"Win64"
+			]
 		},
 		{
 			"Name": "HoloLensTargetPlatform",
 			"Type": "Editor",
-			"LoadingPhase": "Default"
+			"LoadingPhase": "Default",
+			"WhitelistPlatforms": [
+				"HoloLens",
+				"Win64"
+			]
 		}
 	]
 }
\ No newline at end of file
