Fix static analysis warning in HoloLensTargetDevice.cpp

diff --git a/Engine/Plugins/Runtime/AR/Microsoft/HoloLensAR/Source/HoloLensTargetPlatform/Private/HoloLensTargetDevice.cpp b/Engine/Plugins/Runtime/AR/Microsoft/HoloLensAR/Source/HoloLensTargetPlatform/Private/HoloLensTargetDevice.cpp
--- a/Engine/Plugins/Runtime/AR/Microsoft/HoloLensAR/Source/HoloLensTargetPlatform/Private/HoloLensTargetDevice.cpp
+++ b/Engine/Plugins/Runtime/AR/Microsoft/HoloLensAR/Source/HoloLensTargetPlatform/Private/HoloLensTargetDevice.cpp
@@ -8,17 +8,7 @@
 #include "Windows/AllowWindowsPlatformTypes.h"
 #endif
 
-//#include <wrl/client.h>
-//#include <wrl/wrappers/corewrappers.h>
-
-//#include <shlwapi.h>
-//#include <shobjidl.h>
-//#include <AppxPackaging.h>
 #include "Windows/ComPointer.h"
-//#include "ComPointer.h"
-//#include <Windows.Management.Deployment.h>
-
-//#include "D:/release_HoloLens/Engine/Source/Developer/HoloLens/ThirdParty/HoloLensBuildLib/HoloLensBuildLib/HoloLensBuildLib.h"
 #include "HoloLensBuildLib.h"
 
 WindowsMixedReality::HoloLensBuildLib hlLib;
@@ -35,7 +25,11 @@ bool FHoloLensTargetDevice::Launch(const FString& AppId, EBuildConfigurations::T
 
 bool FHoloLensTargetDevice::Run(const FString& ExecutablePath, const FString& Params, uint32* OutProcessId)
 {
-	CoInitialize(nullptr);
+	HRESULT hr = CoInitialize(nullptr);
+	if (FAILED(hr))
+	{
+		UE_LOG(LogTemp, Warning, TEXT("FHoloLensTargetDevice::Run - CoInitialize() failed with hr = 0x(%x)"), hr);
+	}
 
 	// Currently even packaged builds get an exe name in here which kind of works because we 
 	// don't yet support remote deployment and so the loose structure the package was created 
