Fixed FHoloLensARSystem::PinComponentToARPin using a variable before null check

diff --git a/Engine/Plugins/Runtime/AR/Microsoft/HoloLensAR/Source/HoloLensAR/Private/HoloLensARSystemAnchors.cpp b/Engine/Plugins/Runtime/AR/Microsoft/HoloLensAR/Source/HoloLensAR/Private/HoloLensARSystemAnchors.cpp
--- a/Engine/Plugins/Runtime/AR/Microsoft/HoloLensAR/Source/HoloLensAR/Private/HoloLensARSystemAnchors.cpp
+++ b/Engine/Plugins/Runtime/AR/Microsoft/HoloLensAR/Source/HoloLensAR/Private/HoloLensARSystemAnchors.cpp
@@ -47,6 +47,11 @@ UWMRARPin* FHoloLensARSystem::CreateNamedARPin(FName Name, const FTransform& Pin
 
 bool FHoloLensARSystem::PinComponentToARPin(USceneComponent* ComponentToPin, UWMRARPin* Pin)
 {
+	if (Pin == nullptr)
+	{
+		UE_LOG(LogHoloLensAR, Warning, TEXT("PinComponentToWMRAnchorStoreARPin: Pin was null.  Doing nothing."));
+		return false;
+	}
 	if (ComponentToPin == nullptr)
 	{
 		UE_LOG(LogHoloLensAR, Warning, TEXT("PinComponentToWMRAnchorStoreARPin: Tried to pin null component to pin %s.  Doing nothing."), *Pin->GetDebugName().ToString());
