Removed redundant check of Pin pointer

diff --git a/Engine/Plugins/Runtime/AR/Microsoft/HoloLensAR/Source/HoloLensAR/Private/HoloLensARSystemAnchors.cpp b/Engine/Plugins/Runtime/AR/Microsoft/HoloLensAR/Source/HoloLensAR/Private/HoloLensARSystemAnchors.cpp
--- a/Engine/Plugins/Runtime/AR/Microsoft/HoloLensAR/Source/HoloLensAR/Private/HoloLensARSystemAnchors.cpp
+++ b/Engine/Plugins/Runtime/AR/Microsoft/HoloLensAR/Source/HoloLensAR/Private/HoloLensARSystemAnchors.cpp
@@ -58,12 +58,6 @@ bool FHoloLensARSystem::PinComponentToARPin(USceneComponent* ComponentToPin, UWM
 		return false;
 	}
 
-	if (Pin == nullptr)
-	{
-		UE_LOG(LogHoloLensAR, Warning, TEXT("PinComponentToWMRAnchorStoreARPin: Tried to pin component %s to null pin.  Doing nothing."), *ComponentToPin->GetReadableName());
-		return false;
-	}
-
 	{
 		if (UWMRARPin* FindResult = FindPinByComponent(ComponentToPin))
 		{
