Suppress bad static analysis by MSVC

diff --git a/Engine/Plugins/Runtime/AR/Microsoft/HoloLensAR/Source/HoloLensPlatformEditor/Private/HoloLensTargetSettingsCustomization.cpp b/Engine/Plugins/Runtime/AR/Microsoft/HoloLensAR/Source/HoloLensPlatformEditor/Private/HoloLensTargetSettingsCustomization.cpp
--- a/Engine/Plugins/Runtime/AR/Microsoft/HoloLensAR/Source/HoloLensPlatformEditor/Private/HoloLensTargetSettingsCustomization.cpp
+++ b/Engine/Plugins/Runtime/AR/Microsoft/HoloLensAR/Source/HoloLensPlatformEditor/Private/HoloLensTargetSettingsCustomization.cpp
@@ -31,6 +31,8 @@
 
 #pragma comment( lib, "Advapi32.lib" )
 
+#pragma warning( disable : 6387 )
+
 #define LOCTEXT_NAMESPACE "HoloLensTargetSettingsCustomization"
 
 TArray<FHoloLensSDKVersion> FHoloLensSDKVersion::GetSDKVersions()
