Merging: Change: 6680335 UE-75331 Hololens 2 - Packaged FPS template crashes on launch on reflection capture
-Hololens2 is ES3_1 therefore it needs to output EncodedHDR Reflection format.

Merging //UE4/Dev-Microsoft-HoloLens2-Preview to Dev-VR (//UE4/Dev-VR)

diff --git a/Engine/Plugins/Runtime/AR/Microsoft/HoloLensAR/Source/HoloLensTargetPlatform/Private/HoloLensTargetPlatform.cpp b/Engine/Plugins/Runtime/AR/Microsoft/HoloLensAR/Source/HoloLensTargetPlatform/Private/HoloLensTargetPlatform.cpp
--- a/Engine/Plugins/Runtime/AR/Microsoft/HoloLensAR/Source/HoloLensTargetPlatform/Private/HoloLensTargetPlatform.cpp
+++ b/Engine/Plugins/Runtime/AR/Microsoft/HoloLensAR/Source/HoloLensTargetPlatform/Private/HoloLensTargetPlatform.cpp
@@ -88,6 +88,12 @@ bool FHoloLensTargetPlatform::SupportsFeature(ETargetPlatformFeatures Feature) c
 
 #if WITH_ENGINE
 
+void FHoloLensTargetPlatform::GetReflectionCaptureFormats(TArray<FName>& OutFormats) const
+{
+	OutFormats.Add(FName(TEXT("FullHDR")));
+	OutFormats.Add(FName(TEXT("EncodedHDR")));
+}
+
 const FPlatformAudioCookOverrides* FHoloLensTargetPlatform::GetAudioCompressionSettings() const
 {
 	return nullptr;
diff --git a/Engine/Plugins/Runtime/AR/Microsoft/HoloLensAR/Source/HoloLensTargetPlatform/Private/HoloLensTargetPlatform.h b/Engine/Plugins/Runtime/AR/Microsoft/HoloLensAR/Source/HoloLensTargetPlatform/Private/HoloLensTargetPlatform.h
--- a/Engine/Plugins/Runtime/AR/Microsoft/HoloLensAR/Source/HoloLensTargetPlatform/Private/HoloLensTargetPlatform.h
+++ b/Engine/Plugins/Runtime/AR/Microsoft/HoloLensAR/Source/HoloLensTargetPlatform/Private/HoloLensTargetPlatform.h
@@ -66,6 +66,8 @@ public:
 	virtual bool SupportsBuildTarget(EBuildTargets::Type BuildTarget) const override;
 
 #if WITH_ENGINE
+	virtual void GetReflectionCaptureFormats(TArray<FName>& OutFormats) const override;
+
 	virtual const FPlatformAudioCookOverrides* GetAudioCompressionSettings() const override;
 
 	virtual const class FStaticMeshLODSettings& GetStaticMeshLODSettings() const override { return StaticMeshLODSettings; }
