Hooked up the AR delegates for Add/Update/Remove of tracked geometry for HoloLens

diff --git a/Engine/Plugins/Runtime/AR/Microsoft/HoloLensAR/Source/HoloLensAR/Private/HoloLensARSystem.cpp b/Engine/Plugins/Runtime/AR/Microsoft/HoloLensAR/Source/HoloLensAR/Private/HoloLensARSystem.cpp
--- a/Engine/Plugins/Runtime/AR/Microsoft/HoloLensAR/Source/HoloLensAR/Private/HoloLensARSystem.cpp
+++ b/Engine/Plugins/Runtime/AR/Microsoft/HoloLensAR/Source/HoloLensAR/Private/HoloLensARSystem.cpp
@@ -628,6 +628,7 @@ void FHoloLensARSystem::ProcessMeshUpdates_GameThread()
 
 void FHoloLensARSystem::AddOrUpdateMesh(FMeshUpdate* CurrentMesh)
 {
+	bool bIsAdd = false;
 	UARTrackedGeometry* NewUpdatedGeometry = nullptr;
 	UARTrackedGeometry** FoundGeometry = TrackedGeometries.Find(CurrentMesh->Id);
 	if (FoundGeometry == nullptr)
@@ -635,6 +636,7 @@ void FHoloLensARSystem::AddOrUpdateMesh(FMeshUpdate* CurrentMesh)
 		// We haven't seen this one before so add it to our set
 		NewUpdatedGeometry = NewObject<UARTrackedGeometry>();
 		TrackedGeometries.Add(CurrentMesh->Id, NewUpdatedGeometry);
+		bIsAdd = true;
 	}
 	else
 	{
@@ -677,6 +679,15 @@ void FHoloLensARSystem::AddOrUpdateMesh(FMeshUpdate* CurrentMesh)
 		// MRMesh takes ownership of the data in the arrays at this point
 		MRMesh->UpdateMesh(CurrentMesh->Location, CurrentMesh->Rotation, CurrentMesh->Scale, CurrentMesh->Vertices, CurrentMesh->Indices);
 	}
+	// Trigger the proper notification delegate
+	if (bIsAdd)
+	{
+		TriggerOnTrackableAddedDelegates(NewUpdatedGeometry);
+	}
+	else
+	{
+		TriggerOnTrackableUpdatedDelegates(NewUpdatedGeometry);
+	}
 }
 
 void FHoloLensARSystem::ReconcileKnownMeshes(const TArray<FGuid>& KnownMeshes)
@@ -704,6 +715,7 @@ void FHoloLensARSystem::ReconcileKnownMeshes(const TArray<FGuid>& KnownMeshes)
 			}
 
 			TrackedGeometries.Remove(*Iter);
+			TriggerOnTrackableRemovedDelegates(*TrackedGeometry);
 		}
 	}
 	// Update our last know list to the currently known set of meshes
