Merge from Dev-Microsoft-HoloLens2-Preview CL 6682981

Simplifying installed SDK options
MinimumPlatformVersion and MaximumPlatformVersionTested are now moved to AdvancedDisplay so they are hidden by default
MinimumPlatformVersion is now forced to 10.0.18362.0 when HoloLensEngine.ini is created to make sure the package can run on older firmware devices

Removing unused or unimplemented options
bCopyCookedContentForF5Deployment, bExcludeShaderModel4Support, bBuildD3D12RHI, bUseD3D12RHI, bUseAssetPackage are all removed because they aren't ready/working
HoloLensDlc* has been removed since DLC is not and will not be supported

diff --git a/Engine/Plugins/Runtime/AR/Microsoft/HoloLensAR/Source/HoloLensPlatformEditor/Private/HoloLensDlcSettingsCustomization.cpp b/Engine/Plugins/Runtime/AR/Microsoft/HoloLensAR/Source/HoloLensPlatformEditor/Private/HoloLensDlcSettingsCustomization.cpp
deleted file mode 100644
diff --git a/Engine/Plugins/Runtime/AR/Microsoft/HoloLensAR/Source/HoloLensPlatformEditor/Private/HoloLensDlcSettingsCustomization.h b/Engine/Plugins/Runtime/AR/Microsoft/HoloLensAR/Source/HoloLensPlatformEditor/Private/HoloLensDlcSettingsCustomization.h
deleted file mode 100644
diff --git a/Engine/Plugins/Runtime/AR/Microsoft/HoloLensAR/Source/HoloLensPlatformEditor/Private/HoloLensImageResourcesCustomization.cpp b/Engine/Plugins/Runtime/AR/Microsoft/HoloLensAR/Source/HoloLensPlatformEditor/Private/HoloLensImageResourcesCustomization.cpp
--- a/Engine/Plugins/Runtime/AR/Microsoft/HoloLensAR/Source/HoloLensPlatformEditor/Private/HoloLensImageResourcesCustomization.cpp
+++ b/Engine/Plugins/Runtime/AR/Microsoft/HoloLensAR/Source/HoloLensPlatformEditor/Private/HoloLensImageResourcesCustomization.cpp
@@ -19,7 +19,7 @@
 #include "DesktopPlatformModule.h"
 #include "HAL/FileManager.h"
 
-#define LOCTEXT_NAMESPACE "HoloLensDlcImagesCustomization"
+#define LOCTEXT_NAMESPACE "HoloLensImagesCustomization"
 
 namespace
 {
@@ -136,28 +136,6 @@ namespace
 	}
 }
 
-TSharedRef<IPropertyTypeCustomization> FHoloLensDlcImagesCustomization::MakeInstance()
-{
-	return MakeShared<FHoloLensDlcImagesCustomization>();
-}
-
-void FHoloLensDlcImagesCustomization::CustomizeChildren(TSharedRef<IPropertyHandle> InStructPropertyHandle, IDetailChildrenBuilder& ChildBuilder, IPropertyTypeCustomizationUtils& StructCustomizationUtils)
-{
-	TSharedRef<IPropertyHandle> PluginNameProperty = InStructPropertyHandle->GetParentHandle()->GetChildHandle(FName("AppliesToDlcPlugin")).ToSharedRef();
-	FString CurrentPluginName;
-	PluginNameProperty->GetValue(CurrentPluginName);
-
-	TSharedPtr<IPlugin> Plugin = IPluginManager::Get().FindPlugin(CurrentPluginName);
-	if (Plugin.IsValid())
-	{
-		TSharedRef<IPropertyHandle> CultureIdProperty = InStructPropertyHandle->GetParentHandle()->GetChildHandle(FName("CultureId")).ToSharedRef();
-		FString CultureId;
-		CultureIdProperty->GetValue(CultureId);
-
-		AddWidgetForResourceImage(ChildBuilder, Plugin->GetBaseDir(), CultureId, TEXT("StoreLogo"), LOCTEXT("StoreLogo", "Store Logo"), FVector2D(50.0f, 50.0f));
-	}
-}
-
 TSharedRef<IPropertyTypeCustomization> FHoloLensCorePackageImagesCustomization::MakeInstance()
 {
 	return MakeShared<FHoloLensCorePackageImagesCustomization>();
diff --git a/Engine/Plugins/Runtime/AR/Microsoft/HoloLensAR/Source/HoloLensPlatformEditor/Private/HoloLensImageResourcesCustomization.h b/Engine/Plugins/Runtime/AR/Microsoft/HoloLensAR/Source/HoloLensPlatformEditor/Private/HoloLensImageResourcesCustomization.h
--- a/Engine/Plugins/Runtime/AR/Microsoft/HoloLensAR/Source/HoloLensPlatformEditor/Private/HoloLensImageResourcesCustomization.h
+++ b/Engine/Plugins/Runtime/AR/Microsoft/HoloLensAR/Source/HoloLensPlatformEditor/Private/HoloLensImageResourcesCustomization.h
@@ -16,15 +16,3 @@ public:
 	virtual void CustomizeChildren(TSharedRef<IPropertyHandle> InStructPropertyHandle, class IDetailChildrenBuilder& ChildBuilder, IPropertyTypeCustomizationUtils& StructCustomizationUtils) override;
 	// End of IPropertyTypeCustomization interface
 };
-
-class FHoloLensDlcImagesCustomization : public IPropertyTypeCustomization
-{
-public:
-	static TSharedRef<IPropertyTypeCustomization> MakeInstance();
-
-	// IPropertyTypeCustomization interface
-	virtual void CustomizeHeader(TSharedRef<class IPropertyHandle> StructPropertyHandle, class FDetailWidgetRow& HeaderRow, IPropertyTypeCustomizationUtils& StructCustomizationUtils) override {}
-
-	virtual void CustomizeChildren(TSharedRef<IPropertyHandle> InStructPropertyHandle, class IDetailChildrenBuilder& ChildBuilder, IPropertyTypeCustomizationUtils& StructCustomizationUtils) override;
-	// End of IPropertyTypeCustomization interface
-};
\ No newline at end of file
diff --git a/Engine/Plugins/Runtime/AR/Microsoft/HoloLensAR/Source/HoloLensPlatformEditor/Private/HoloLensLocalizedResources.h b/Engine/Plugins/Runtime/AR/Microsoft/HoloLensAR/Source/HoloLensPlatformEditor/Private/HoloLensLocalizedResources.h
--- a/Engine/Plugins/Runtime/AR/Microsoft/HoloLensAR/Source/HoloLensPlatformEditor/Private/HoloLensLocalizedResources.h
+++ b/Engine/Plugins/Runtime/AR/Microsoft/HoloLensAR/Source/HoloLensPlatformEditor/Private/HoloLensLocalizedResources.h
@@ -25,30 +25,12 @@ struct FHoloLensCorePackageStringResources
 	FString ApplicationDescription;
 };
 
-USTRUCT()
-struct FHoloLensDlcStringResources
-{
-	GENERATED_BODY()
-
-	UPROPERTY(EditAnywhere, config, Category = Packaging)
-	FString PackageDisplayName;
-
-	UPROPERTY(EditAnywhere, config, Category = Packaging)
-	FString PackageDescription;
-};
-
 USTRUCT()
 struct FHoloLensCorePackageImageResources
 {
 	GENERATED_BODY()
 };
 
-USTRUCT()
-struct FHoloLensDlcImageResources
-{
-	GENERATED_BODY()
-};
-
 USTRUCT()
 struct FHoloLensCorePackageLocalizedResources
 {
@@ -63,21 +45,3 @@ struct FHoloLensCorePackageLocalizedResources
 	UPROPERTY(EditAnywhere, config, Category = Packaging)
 	FHoloLensCorePackageImageResources Images;
 };
-
-USTRUCT()
-struct FHoloLensDlcLocalizedResources
-{
-	GENERATED_BODY()
-
-	UPROPERTY(EditAnywhere, config, Category = Packaging)
-	FString AppliesToDlcPlugin;
-
-	UPROPERTY(EditAnywhere, config, Category = Packaging)
-	FString CultureId;
-
-	UPROPERTY(EditAnywhere, config, Category = Packaging, Meta = (ShowOnlyInnerProperties))
-	FHoloLensDlcStringResources Strings;
-
-	UPROPERTY(EditAnywhere, config, Category = Packaging)
-	FHoloLensDlcImageResources Images;
-};
\ No newline at end of file
diff --git a/Engine/Plugins/Runtime/AR/Microsoft/HoloLensAR/Source/HoloLensPlatformEditor/Private/HoloLensPlatformEditorModule.cpp b/Engine/Plugins/Runtime/AR/Microsoft/HoloLensAR/Source/HoloLensPlatformEditor/Private/HoloLensPlatformEditorModule.cpp
--- a/Engine/Plugins/Runtime/AR/Microsoft/HoloLensAR/Source/HoloLensPlatformEditor/Private/HoloLensPlatformEditorModule.cpp
+++ b/Engine/Plugins/Runtime/AR/Microsoft/HoloLensAR/Source/HoloLensPlatformEditor/Private/HoloLensPlatformEditorModule.cpp
@@ -3,7 +3,6 @@
 #include "PropertyEditorModule.h"
 #include "HoloLensTargetSettingsCustomization.h"
 #include "HoloLensLocalizedResourcesCustomization.h"
-#include "HoloLensDlcSettingsCustomization.h"
 #include "HoloLensImageResourcesCustomization.h"
 #include "ISettingsModule.h"
 #include "Modules/ModuleManager.h"
@@ -58,25 +57,10 @@ public:
 			FOnGetPropertyTypeCustomizationInstance::CreateStatic(&FHoloLensLocalizedResourcesCustomization::MakeInstance)
 			);
 
-		PropertyModule.RegisterCustomPropertyTypeLayout(
-			"HoloLensDlcLocalizedResources",
-			FOnGetPropertyTypeCustomizationInstance::CreateStatic(&FHoloLensLocalizedResourcesCustomization::MakeInstance)
-			);
-
-		PropertyModule.RegisterCustomPropertyTypeLayout(
-			"HoloLensDlcSettings",
-			FOnGetPropertyTypeCustomizationInstance::CreateStatic(&FHoloLensDlcSettingsCustomization::MakeInstance)
-			);
-
 		PropertyModule.RegisterCustomPropertyTypeLayout(
 			"HoloLensCorePackageImageResources",
 			FOnGetPropertyTypeCustomizationInstance::CreateStatic(&FHoloLensCorePackageImagesCustomization::MakeInstance)
 			);
-
-		PropertyModule.RegisterCustomPropertyTypeLayout(
-			"HoloLensDlcImageResources",
-			FOnGetPropertyTypeCustomizationInstance::CreateStatic(&FHoloLensDlcImagesCustomization::MakeInstance)
-			);
 	}
 
 	virtual void ShutdownModule() override
diff --git a/Engine/Plugins/Runtime/AR/Microsoft/HoloLensAR/Source/HoloLensPlatformEditor/Private/HoloLensTargetSettings.h b/Engine/Plugins/Runtime/AR/Microsoft/HoloLensAR/Source/HoloLensPlatformEditor/Private/HoloLensTargetSettings.h
--- a/Engine/Plugins/Runtime/AR/Microsoft/HoloLensAR/Source/HoloLensPlatformEditor/Private/HoloLensTargetSettings.h
+++ b/Engine/Plugins/Runtime/AR/Microsoft/HoloLensAR/Source/HoloLensPlatformEditor/Private/HoloLensTargetSettings.h
@@ -10,22 +10,6 @@
 #include "HoloLensLocalizedResources.h"
 #include "HoloLensTargetSettings.generated.h"
 
-USTRUCT()
-struct FHoloLensDLCSettings
-{
-public:
-	GENERATED_BODY()
-
-	UPROPERTY(EditAnywhere, config, Category = DLC)
-	FString PluginName;
-
-	UPROPERTY(EditAnywhere, config, Category = DLC, Meta = (DisplayName = "Package/Identity/Name"))
-	FString PackageIdentityName;
-
-	UPROPERTY(EditAnywhere, config, Category = DLC, Meta = (DisplayName = "Package/Identity/Version"))
-	FString PackageIdentityVersion;
-};
-
 /**
  * Implements the settings for the HoloLens target platform.
  */
@@ -56,34 +40,6 @@ public:
 	UPROPERTY(EditAnywhere, Config, Category="HoloLens", Meta=(DisplayName="Build for HoloLens Device"))
 	bool bBuildForDevice = true;
 
-	/**
-	 * HoloLens includes SM4 support (D3D Feature level 10) by default.  This setting excludes this
-	 * support, possibly reducing package size, but limiting the range of devices the package
-	 * will run on.  Currently you must leave this DISABLED if you wish to run as HoloLens on Xbox One.
-	 */
-	UPROPERTY(EditAnywhere, config, Category = Rendering, Meta=(DisplayName="Exclude Shader Model 4 Support"))
-	bool bExcludeShaderModel4Support;
-
-	/**
-	 * Selects whether to build the D3D12 RHI module.  The D3D12 RHI may have requirements for minimum supported
-	 * Windows SDK version and compiler toolset beyond the defaults for HoloLens support.  D3D11 support is always
-	 * included in the final exe regardless of this setting.
-	 */
-	UPROPERTY(EditAnywhere, config, Category = Rendering, Meta = (DisplayName = "Build with D3D12 support"))
-	bool bBuildD3D12RHI;
-
-	/**
-	 * Selects whether to use the D3D12 RHI by default at runtime.  Requires that the exe is build with D3D12 support enabled.
-	 */
-	UPROPERTY(EditAnywhere, config, Category = Rendering, Meta=(DisplayName = "Run using D3D12 RHI", editcondition="bBuildD3D12RHI"))
-	bool bUseD3D12RHI;
-
-	/**
-	 * Whether cooked content should be copied into the loose package image generated by deploy from Visual Studio.
-	 */
-	UPROPERTY(EditAnywhere, config, Category = "Toolchain", Meta = (DisplayName = "Copy Cooked Content for F5 Deployment"))
-	bool bCopyCookedContentForF5Deployment;
-
 	/**
 	 * Controls whether to use the retail Windows Store environment for license checks.  This must be turned on
 	 * when building for submission to the Windows Store, or when sideloading outside of Developer Mode.  Note,
@@ -115,7 +71,7 @@ public:
 	 * Minimum version of the HoloLens platform required to run this title.
 	 * It will not be possible to deploy the build on earlier versions.
 	 */
-	UPROPERTY(EditAnywhere, config, Category = "OS Info", Meta = (DisplayName = "Minimum supported platform version"))
+	UPROPERTY(EditAnywhere, config, Category = "OS Info", AdvancedDisplay, Meta = (DisplayName = "Minimum supported platform version"))
 	FString MinimumPlatformVersion;
 
 	/**
@@ -123,7 +79,7 @@ public:
 	 * title is known to work as expected.  When deployed to later versions
 	 * the title will experience the runtime behavior of the version given here.
 	 */
-	UPROPERTY(EditAnywhere, config, Category = "OS Info", Meta = (DisplayName = "Maximum tested platform version"))
+	UPROPERTY(EditAnywhere, config, Category = "OS Info", AdvancedDisplay, Meta = (DisplayName = "Maximum tested platform version"))
 	FString MaximumPlatformVersionTested;
 
 	/** Used by the HoloLens to indicate triangle density when generating meshes. Defaults to 500 per cubic meter */
@@ -134,12 +90,6 @@ public:
 	UPROPERTY(EditAnywhere, Config, Category = "Spatial Mapping")
 	float SpatialMeshingVolumeSize;
 
-	/**
-	 * Package Assets into a separate file for reducing upload size.
-	 */
-	UPROPERTY(EditAnywhere, config, Category = Packaging, Meta = (DisplayName = "Package assets into a separate file"))
-	bool bUseAssetPackage = false;
-
 	/** The compiler version to use for this project. May be different to the chosen IDE. */
 	UPROPERTY(EditAnywhere, config, Category = "Toolchain", Meta = (DisplayName = "Compiler Version"))
 	ECompilerVersion CompilerVersion;
@@ -177,12 +127,6 @@ public:
 	UPROPERTY(EditAnywhere, config, Category = Capabilities)
 	bool bSetDefaultCapabilities = true;
 
-	UPROPERTY(EditAnywhere, config, Category = DLC)
-	TArray<FHoloLensDLCSettings> DLCStoreMapping;
-
-	UPROPERTY(EditAnywhere, config, Category = DLC)
-	TArray<FHoloLensDlcLocalizedResources> DLCPerCultureResources;
-
 	/** Which of the currently enabled spatialization plugins to use on Windows. */
 	UPROPERTY(config, EditAnywhere, Category = "Audio")
 	FString SpatializationPlugin;
diff --git a/Engine/Plugins/Runtime/AR/Microsoft/HoloLensAR/Source/HoloLensPlatformEditor/Private/HoloLensTargetSettingsCustomization.cpp b/Engine/Plugins/Runtime/AR/Microsoft/HoloLensAR/Source/HoloLensPlatformEditor/Private/HoloLensTargetSettingsCustomization.cpp
--- a/Engine/Plugins/Runtime/AR/Microsoft/HoloLensAR/Source/HoloLensPlatformEditor/Private/HoloLensTargetSettingsCustomization.cpp
+++ b/Engine/Plugins/Runtime/AR/Microsoft/HoloLensAR/Source/HoloLensPlatformEditor/Private/HoloLensTargetSettingsCustomization.cpp
@@ -37,6 +37,8 @@
 
 TArray<FHoloLensSDKVersion> FHoloLensSDKVersion::GetSDKVersions()
 {
+	bool bFound18362 = false;
+
 	static TArray<FHoloLensSDKVersion> SDKVersions;
 	if (SDKVersions.Num() == 0)
 	{
@@ -62,6 +64,10 @@ TArray<FHoloLensSDKVersion> FHoloLensSDKVersion::GetSDKVersions()
 						{
 							new(SDKVersions) FHoloLensSDKVersion(FString(dwSubKeyNameLen, sSubKeyName), version[0], version[1], version[2], version[3]);
 						}
+						if (version[2] == 18362)
+						{
+							bFound18362 = true;
+						}
 					}
 					dwSubKeyNameLen = _countof(sSubKeyName);
 					dwIndex++;
@@ -69,7 +75,16 @@ TArray<FHoloLensSDKVersion> FHoloLensSDKVersion::GetSDKVersions()
 				RegCloseKey(hKey);
 			}
 		}
+
+		// 18362 is required as the official minimum SDK released by Microsoft for all needed HL2 features to work
+		// If it's not found, we still need it as a hidden SDK version so that we can set MinimumPlatformVersion correctly to
+		// allow the package to run on lower installed firmware revisions
+		if (!bFound18362)
+		{
+			new(SDKVersions) FHoloLensSDKVersion(FString(L"10.0.18362.0"), 10, 0, 18362, 0);
+		}
 	}
+
 	return SDKVersions;
 }
 
@@ -107,9 +122,10 @@ void FHoloLensTargetSettingsCustomization::CustomizeDetails(IDetailLayoutBuilder
 
 	// Add UI for select min/max platform version.
 	TSharedRef<IPropertyHandle> MinVersionProperty = DetailBuilder.GetProperty("MinimumPlatformVersion");
-	AddWidgetForPlatformVersion(DetailBuilder, MinVersionProperty);
+	// 10.0.18362.0 is the minimum 'officially' released SDK from MS that supports all the HoloLens 2 features that are needed
+	AddWidgetForPlatformVersion(DetailBuilder, MinVersionProperty, nullptr, true, FString(L"10.0.18362.0"));
 	TSharedRef<IPropertyHandle> MaxVersionProperty = DetailBuilder.GetProperty("MaximumPlatformVersionTested");
-	AddWidgetForPlatformVersion(DetailBuilder, MaxVersionProperty);
+	AddWidgetForPlatformVersion(DetailBuilder, MaxVersionProperty, nullptr, true);
 
 	// Add UI to select signing certificate
 	IDetailCategoryBuilder& PackagingCategoryBuilder = DetailBuilder.EditCategory(FName("Packaging"));
@@ -315,12 +331,12 @@ FString FHoloLensTargetSettingsCustomization::GetNameForSigningCertificate(const
 	return CertificateName;
 }
 
-void FHoloLensTargetSettingsCustomization::AddWidgetForPlatformVersion(IDetailLayoutBuilder& DetailBuilder, TSharedRef<IPropertyHandle> PropertyHandle, TSharedPtr<STextComboBox>* OutVersionSelector)
+void FHoloLensTargetSettingsCustomization::AddWidgetForPlatformVersion(IDetailLayoutBuilder& DetailBuilder, TSharedRef<IPropertyHandle> PropertyHandle, TSharedPtr<STextComboBox>* OutVersionSelector, bool bForAdvanced, const FString& SpecificVersionOverrideString)
 {
 	IDetailCategoryBuilder& VersionCategoryBuilder = DetailBuilder.EditCategory(FName(*PropertyHandle->GetMetaData("Category")));
 	DetailBuilder.HideProperty(PropertyHandle);
 
-	FString CurrentSelectedVersion;
+	FString CurrentSelectedVersion = SpecificVersionOverrideString;
 	PropertyHandle->GetValue(CurrentSelectedVersion);
 
 	// Default to latest version when not set explicitly.  With Windows 10 automatic updates
@@ -346,7 +362,7 @@ void FHoloLensTargetSettingsCustomization::AddWidgetForPlatformVersion(IDetailLa
 	TSharedPtr<STextComboBox> Unused;
 	TSharedPtr<STextComboBox>& VersionSelectorRef = OutVersionSelector != nullptr ? *OutVersionSelector : Unused;
 
-	VersionCategoryBuilder.AddCustomRow(PropertyHandle->GetPropertyDisplayName())
+	VersionCategoryBuilder.AddCustomRow(PropertyHandle->GetPropertyDisplayName(), bForAdvanced)
 	.NameContent()
 	[
 		PropertyHandle->CreatePropertyNameWidget()
diff --git a/Engine/Plugins/Runtime/AR/Microsoft/HoloLensAR/Source/HoloLensPlatformEditor/Private/HoloLensTargetSettingsCustomization.h b/Engine/Plugins/Runtime/AR/Microsoft/HoloLensAR/Source/HoloLensPlatformEditor/Private/HoloLensTargetSettingsCustomization.h
--- a/Engine/Plugins/Runtime/AR/Microsoft/HoloLensAR/Source/HoloLensPlatformEditor/Private/HoloLensTargetSettingsCustomization.h
+++ b/Engine/Plugins/Runtime/AR/Microsoft/HoloLensAR/Source/HoloLensPlatformEditor/Private/HoloLensTargetSettingsCustomization.h
@@ -23,7 +23,7 @@ public:
 private:
 
 	void AddWidgetForCapability(IDetailLayoutBuilder& DetailBuilder, TSharedRef<IPropertyHandle> CapabilityList, const FString& CapabilityName, const FText& CapabilityCaption, const FText& CapabilityTooltip, bool bForAdvanced);
-	void AddWidgetForPlatformVersion(IDetailLayoutBuilder& DetailBuilder, TSharedRef<IPropertyHandle> PropertyHandle, TSharedPtr<STextComboBox>* OutVersionSelector = nullptr);
+	void AddWidgetForPlatformVersion(IDetailLayoutBuilder& DetailBuilder, TSharedRef<IPropertyHandle> PropertyHandle, TSharedPtr<STextComboBox>* OutVersionSelector = nullptr, bool bForAdvanced = false, const FString& SpecificVersionOverrideString = FString());
 	void AddWidgetForTargetDeviceFamily(IDetailLayoutBuilder& DetailBuilder, TSharedRef<IPropertyHandle> PropertyHandle);
 	static FString GetNameForSigningCertificate(const FString &CertificatePath);
 	void InitSupportedPlatformVersions();
diff --git a/Engine/Plugins/Runtime/AR/Microsoft/HoloLensAR/Source/HoloLensTargetPlatform/Private/HoloLensTargetPlatform.cpp b/Engine/Plugins/Runtime/AR/Microsoft/HoloLensAR/Source/HoloLensTargetPlatform/Private/HoloLensTargetPlatform.cpp
--- a/Engine/Plugins/Runtime/AR/Microsoft/HoloLensAR/Source/HoloLensTargetPlatform/Private/HoloLensTargetPlatform.cpp
+++ b/Engine/Plugins/Runtime/AR/Microsoft/HoloLensAR/Source/HoloLensTargetPlatform/Private/HoloLensTargetPlatform.cpp
@@ -101,17 +101,13 @@ const FPlatformAudioCookOverrides* FHoloLensTargetPlatform::GetAudioCompressionS
 
 void FHoloLensTargetPlatform::GetTextureFormats(const UTexture* InTexture, TArray<FName>& OutFormats) const
 {
-	bool bExcludeShaderModel4Support = false;
-	GConfig->GetBool(TEXT("/Script/HoloLensPlatformEditor.HoloLensTargetSettings"), TEXT("bExcludeShaderModel4Support"), bExcludeShaderModel4Support, GEngineIni);
-	FName TextureFormatName = GetDefaultTextureFormatName(this, InTexture, EngineSettings, bExcludeShaderModel4Support);
+	FName TextureFormatName = GetDefaultTextureFormatName(this, InTexture, EngineSettings, false);
 	OutFormats.Add(TextureFormatName);
 }
 
 void FHoloLensTargetPlatform::GetAllTextureFormats(TArray<FName>& OutFormats) const
 {
-	bool bExcludeShaderModel4Support = false;
-	GConfig->GetBool(TEXT("/Script/HoloLensPlatformEditor.HoloLensTargetSettings"), TEXT("bExcludeShaderModel4Support"), bExcludeShaderModel4Support, GEngineIni);
-	GetAllDefaultTextureFormats(this, OutFormats, bExcludeShaderModel4Support);
+	GetAllDefaultTextureFormats(this, OutFormats, false);
 }
 
 static FName NAME_PCD3D_ES3_1(TEXT("PCD3D_ES31"));
@@ -128,14 +124,8 @@ void FHoloLensTargetPlatform::GetAllPossibleShaderFormats(TArray<FName>& OutForm
 void FHoloLensTargetPlatform::GetAllTargetedShaderFormats(TArray<FName>& OutFormats) const
 {
 	OutFormats.AddUnique(NAME_PCD3D_ES3_1);
-	
-	bool bExcludeShaderModel4Support = false;
-	GConfig->GetBool(TEXT("/Script/HoloLensPlatformEditor.HoloLensTargetSettings"), TEXT("bExcludeShaderModel4Support"), bExcludeShaderModel4Support, GEngineIni);
 	OutFormats.AddUnique(NAME_PCD3D_SM5);
-	if (!bExcludeShaderModel4Support)
-	{
-		OutFormats.AddUnique(NAME_PCD3D_SM4);
-	}
+	OutFormats.AddUnique(NAME_PCD3D_SM4);
 }
 
 #endif
diff --git a/Engine/Source/Programs/AutomationTool/HoloLens/HoloLensPlatform.Automation.cs b/Engine/Source/Programs/AutomationTool/HoloLens/HoloLensPlatform.Automation.cs
--- a/Engine/Source/Programs/AutomationTool/HoloLens/HoloLensPlatform.Automation.cs
+++ b/Engine/Source/Programs/AutomationTool/HoloLens/HoloLensPlatform.Automation.cs
@@ -792,24 +792,10 @@ LogWarning("PackagePakFiles intermediate dir {0}", IntermediateDirectory);
 
 			string OutputNameBase = Params.HasDLCName ? Params.DLCFile.GetFileNameWithoutExtension() : Params.ShortProjectName;
 			string OutputAppX = Path.Combine(SC.StageDirectory.FullName, OutputNameBase + ".appxbundle");
-			bool SeparateAssetPackaging = false;
 
-			ConfigHierarchy PlatformEngineConfig = null;
-			if (Params.EngineConfigs.TryGetValue(PlatformType, out PlatformEngineConfig))
-			{
-				PlatformEngineConfig.GetBool("/Script/UWPPlatformEditor.UWPTargetSettings", "bUseAssetPackage", out SeparateAssetPackaging);
-			}
-
-			if (SeparateAssetPackaging)
-			{
-				PackagePakFiles(Params, SC, OutputNameBase);
-			}
-			else
-			{
-				UpdateCodePackagesWithData(Params, SC, OutputNameBase);
-			}
+			UpdateCodePackagesWithData(Params, SC, OutputNameBase);
 
-			MakeBundle(Params, SC, OutputNameBase, SeparateAssetPackaging);
+			MakeBundle(Params, SC, OutputNameBase, false);
 
 			CopyVCLibs(Params, SC);
 
diff --git a/Engine/Source/Programs/UnrealBuildTool/Platform/HoloLens/HoloLensDeploy.cs b/Engine/Source/Programs/UnrealBuildTool/Platform/HoloLens/HoloLensDeploy.cs
--- a/Engine/Source/Programs/UnrealBuildTool/Platform/HoloLens/HoloLensDeploy.cs
+++ b/Engine/Source/Programs/UnrealBuildTool/Platform/HoloLens/HoloLensDeploy.cs
@@ -243,20 +243,10 @@ namespace UnrealBuildTool
 			var EngineRoot = UnrealBuildTool.RootDirectory;
 			var AddedFiles = new Dictionary<string, string>();
 			bool PackageFileNeedToBeUpdated = !File.Exists(OutputAppX);
-			bool bUseAssetPackage = false;
 
 			DateTime AppXTime = DateTime.Now;
 
-			ConfigHierarchy EngineIni = ConfigCache.ReadHierarchy(ConfigHierarchyType.Engine, DirectoryReference.FromFile(Receipt.ProjectFile), Receipt.Platform);
-			if (EngineIni != null)
-			{
-				EngineIni.GetBool("/Script/UWPPlatformEditor.UWPTargetSettings", "bUseAssetPackage", out bUseAssetPackage);
-			}
-
-			if (!bUseAssetPackage)
-			{
-				PackageFileNeedToBeUpdated = true;
-			}
+			PackageFileNeedToBeUpdated = true;
 
 			if(!PackageFileNeedToBeUpdated)
 			{
@@ -452,39 +442,7 @@ namespace UnrealBuildTool
 			}
 			File.WriteAllText(MapFilename, AppXRecipeBuiltFiles.ToString(), Encoding.UTF8);
 
-			if (bUseAssetPackage)
-			{
-				string MakeAppXCommandLine = string.Format(@"pack /o /f ""{0}"" /p ""{1}""", MapFilename, OutputAppX);
-
-				var StartInfo = new ProcessStartInfo(MakeAppXPath.FullName, MakeAppXCommandLine);
-				StartInfo.UseShellExecute = false;
-				StartInfo.CreateNoWindow = true;
-				var ExitCode = Utils.RunLocalProcessAndPrintfOutput(StartInfo);
-				if (ExitCode < 0)
-				{
-					throw new BuildException("Failed to generate AppX file.  See log for details.");
-				}
-
-				if (File.Exists(SigningCertificatePath))
-				{
-					string SignToolCommandLine = string.Format(@"sign /a /f ""{0}"" /fd SHA256 ""{1}""", SigningCertificatePath, OutputAppX);
-					StartInfo = new ProcessStartInfo(SignToolPath.FullName, SignToolCommandLine);
-					StartInfo.UseShellExecute = false;
-					StartInfo.CreateNoWindow = true;
-					ExitCode = Utils.RunLocalProcessAndPrintfOutput(StartInfo);
-					if (ExitCode < 0)
-					{
-						throw new BuildException("Failed to generate AppX file.  See log for details.");
-					}
-				}
-
-				Log.TraceInformation("AppX successfully packaged to \'{0}.appx\'", OutputName);
-				NewReceipt.BuildProducts.Add(new BuildProduct(new FileReference(OutputAppX), BuildProductType.Package));
-			}
-			else
-			{			
-				NewReceipt.BuildProducts.Add(new BuildProduct(new FileReference(MapFilename), BuildProductType.MapFile));
-			}
+			NewReceipt.BuildProducts.Add(new BuildProduct(new FileReference(MapFilename), BuildProductType.MapFile));
 		}
 
 
diff --git a/Engine/Source/Programs/UnrealBuildTool/Platform/HoloLens/UEBuildHoloLens.cs b/Engine/Source/Programs/UnrealBuildTool/Platform/HoloLens/UEBuildHoloLens.cs
--- a/Engine/Source/Programs/UnrealBuildTool/Platform/HoloLens/UEBuildHoloLens.cs
+++ b/Engine/Source/Programs/UnrealBuildTool/Platform/HoloLens/UEBuildHoloLens.cs
@@ -42,12 +42,6 @@ namespace UnrealBuildTool
 		[ConfigFile(ConfigHierarchyType.Engine, "/Script/HoloLensPlatformEditor.HoloLensTargetSettings", "bBuildForRetailWindowsStore")]
 		public bool bBuildForRetailWindowsStore = false;
 
-		/// <summary>
-		/// Controls whether the D3D12 RHI should be included in the build.
-		/// </summary>
-		[ConfigFile(ConfigHierarchyType.Engine, "/Script/HoloLensPlatformEditor.HoloLensTargetSettings", "bBuildD3D12RHI")]
-		public bool bBuildD3D12RHI = true;
-
 		/// <summary>
 		/// Contains the specific version of the Windows 10 SDK that we will build against. If empty, it will be "Latest"
 		/// </summary>
@@ -103,11 +97,6 @@ namespace UnrealBuildTool
 			get { return Inner.bBuildForRetailWindowsStore; }
 		}
 
-		public bool bBuildD3D12RHI
-		{
-			get { return Inner.bBuildD3D12RHI; }
-		}
-
 		public Version Win10SDKVersion
 		{
 			get { return Inner.Win10SDKVersion; }
@@ -244,12 +233,6 @@ namespace UnrealBuildTool
 			}
 
 			HoloLensExports.InitWindowsSdkToolPath(Target.HoloLensPlatform.Win10SDKVersion.ToString());
-
-			if (Target.HoloLensPlatform.bBuildD3D12RHI && Target.HoloLensPlatform.Win10SDKVersion < MinimumSDKVersionForD3D12RHI)
-			{
-				Log.TraceWarning("Ignoring HoloLens 'Build with D3D12 support' flag: the D3D12 RHI requires at least the {0} SDK.", MinimumSDKVersionForD3D12RHI);
-				Target.HoloLensPlatform.bBuildD3D12RHI = false;
-			}
 		}
 
 		/// <summary>
@@ -629,14 +612,7 @@ namespace UnrealBuildTool
 				CompileEnvironment.Definitions.Add("USING_RETAIL_WINDOWS_STORE=0");
 			}
 
-			if (Target.HoloLensPlatform.bBuildD3D12RHI)
-			{
-				CompileEnvironment.Definitions.Add("WITH_D3D12_RHI=1");
-			}
-			else
-			{
-				CompileEnvironment.Definitions.Add("WITH_D3D12_RHI=0");
-			}
+			CompileEnvironment.Definitions.Add("WITH_D3D12_RHI=0");
 
 			LinkEnvironment.AdditionalArguments += "/NODEFAULTLIB";
 			//CompileEnvironment.AdditionalArguments += " /showIncludes";
diff --git a/Engine/Source/Runtime/D3D12RHI/D3D12RHI.Build.cs b/Engine/Source/Runtime/D3D12RHI/D3D12RHI.Build.cs
--- a/Engine/Source/Runtime/D3D12RHI/D3D12RHI.Build.cs
+++ b/Engine/Source/Runtime/D3D12RHI/D3D12RHI.Build.cs
@@ -46,22 +46,13 @@ public class D3D12RHI : ModuleRules
 		{
 			AddEngineThirdPartyPrivateStaticDependencies(Target, "DX12");
 			AddEngineThirdPartyPrivateStaticDependencies(Target, "DX11");
-			// @ATG_CHANGE : BEGIN HoloLens support
-			if (Target.Platform == UnrealTargetPlatform.HoloLens)
-			{
-				if (!Target.HoloLensPlatform.bBuildD3D12RHI)
-				{
-					//Log.TraceWarning("D3D12 RHI is being built, but HoloLens build settings indicate that it should not be.  Depending on your Windows SDK environment this may cause build errors.  Check build.cs files for dependencies.");
-				}
-			}
-			else
-			{
+            if (Target.Platform != UnrealTargetPlatform.HoloLens)
+            {
                 AddEngineThirdPartyPrivateStaticDependencies(Target, "NVAPI");
 				AddEngineThirdPartyPrivateStaticDependencies(Target, "AMD_AGS");
             	AddEngineThirdPartyPrivateStaticDependencies(Target, "NVAftermath");
             	AddEngineThirdPartyPrivateStaticDependencies(Target, "IntelMetricsDiscovery");
             }
-            // @ATG_CHANGE : END
 		}
     }
 }
diff --git a/Engine/Source/Runtime/Launch/Launch.Build.cs b/Engine/Source/Runtime/Launch/Launch.Build.cs
--- a/Engine/Source/Runtime/Launch/Launch.Build.cs
+++ b/Engine/Source/Runtime/Launch/Launch.Build.cs
@@ -87,11 +87,6 @@ public class Launch : ModuleRules
 				DynamicallyLoadedModuleNames.Add("D3D11RHI");
 				DynamicallyLoadedModuleNames.Add("XAudio2");
 				DynamicallyLoadedModuleNames.Add("AudioMixerXAudio2");
-
-				if (Target.HoloLensPlatform.bBuildD3D12RHI)
-				{
-					DynamicallyLoadedModuleNames.Add("D3D12RHI");
-				}
 			}
 			// @ATG_CHANGE : END
 			else if (Target.Platform == UnrealTargetPlatform.Mac)
diff --git a/Engine/Source/Runtime/RHI/RHI.Build.cs b/Engine/Source/Runtime/RHI/RHI.Build.cs
--- a/Engine/Source/Runtime/RHI/RHI.Build.cs
+++ b/Engine/Source/Runtime/RHI/RHI.Build.cs
@@ -29,11 +29,6 @@ public class RHI : ModuleRules
 				if ((Target.Platform == UnrealTargetPlatform.HoloLens))
 				{
 					DynamicallyLoadedModuleNames.Add("D3D11RHI");
-
-					if (Target.HoloLensPlatform.bBuildD3D12RHI)
-					{
-						DynamicallyLoadedModuleNames.Add("D3D12RHI");
-					}
 				}
 				// @ATG_CHANGE : END
 
