Updated the target platform settings to use the new platform header name

diff --git a/Engine/Plugins/Runtime/AR/Microsoft/HoloLensAR/Source/HoloLensTargetPlatform/Private/HoloLensTargetPlatform.h b/Engine/Plugins/Runtime/AR/Microsoft/HoloLensAR/Source/HoloLensTargetPlatform/Private/HoloLensTargetPlatform.h
--- a/Engine/Plugins/Runtime/AR/Microsoft/HoloLensAR/Source/HoloLensTargetPlatform/Private/HoloLensTargetPlatform.h
+++ b/Engine/Plugins/Runtime/AR/Microsoft/HoloLensAR/Source/HoloLensTargetPlatform/Private/HoloLensTargetPlatform.h
@@ -7,7 +7,7 @@
 #pragma once
 
 #include "Common/TargetPlatformBase.h"
-#include "Runtime/Core/Public/HoloLens/HoloLensProperties.h"
+#include "Runtime/Core/Public/HoloLens/HoloLensPlatformProperties.h"
 #include "Misc/ConfigCacheIni.h"
 #include "HoloLensTargetDevice.h"
 #include "Misc/ScopeLock.h"
