UE-75514 Warning: Failed to read file '../../../../../<Project>/Build/HoloLens/SigningCertificate.pfx' error when opening project
-Suppress the file not found warning from inside the file loading function.  It is expected that this file may be missing and we have specific code to handle that case and report it to the user.

diff --git a/Engine/Plugins/Runtime/AR/Microsoft/HoloLensAR/Source/HoloLensPlatformEditor/Private/HoloLensTargetSettingsCustomization.cpp b/Engine/Plugins/Runtime/AR/Microsoft/HoloLensAR/Source/HoloLensPlatformEditor/Private/HoloLensTargetSettingsCustomization.cpp
--- a/Engine/Plugins/Runtime/AR/Microsoft/HoloLensAR/Source/HoloLensPlatformEditor/Private/HoloLensTargetSettingsCustomization.cpp
+++ b/Engine/Plugins/Runtime/AR/Microsoft/HoloLensAR/Source/HoloLensPlatformEditor/Private/HoloLensTargetSettingsCustomization.cpp
@@ -638,7 +638,7 @@ void FHoloLensTargetSettingsCustomization::LoadAndValidateSigningCertificate()
 
 	FString CertificatePath = FPaths::ProjectDir() / TEXT("Build") / TEXT("HoloLens") / TEXT("SigningCertificate.pfx");
 	TArray<uint8> CertBytes;
-	if (FFileHelper::LoadFileToArray(CertBytes, *CertificatePath))
+	if (FFileHelper::LoadFileToArray(CertBytes, *CertificatePath, FILEREAD_Silent))
 	{
 		CRYPT_DATA_BLOB CertBlob;
 		CertBlob.cbData = CertBytes.Num();
