Moved the QR code updates for HoloLens to the game thread

diff --git a/Engine/Plugins/Runtime/AR/Microsoft/HoloLensAR/Source/HoloLensAR/Private/HoloLensARSystem.cpp b/Engine/Plugins/Runtime/AR/Microsoft/HoloLensAR/Source/HoloLensAR/Private/HoloLensARSystem.cpp
--- a/Engine/Plugins/Runtime/AR/Microsoft/HoloLensAR/Source/HoloLensAR/Private/HoloLensARSystem.cpp
+++ b/Engine/Plugins/Runtime/AR/Microsoft/HoloLensAR/Source/HoloLensAR/Private/HoloLensARSystem.cpp
@@ -18,6 +18,9 @@
 
 DECLARE_CYCLE_STAT(TEXT("OnCameraImageReceived"), STAT_FHoloLensARSystem_OnCameraImageReceived, STATGROUP_HOLOLENS);
 DECLARE_CYCLE_STAT(TEXT("Process Mesh Updates"), STAT_FHoloLensARSystem_ProcessMeshUpdates, STATGROUP_HOLOLENS);
+DECLARE_CYCLE_STAT(TEXT("QR Code Added"), STAT_FHoloLensARSystem_QRCodeAdded_GameThread, STATGROUP_HOLOLENS);
+DECLARE_CYCLE_STAT(TEXT("QR Code Updated"), STAT_FHoloLensARSystem_QRCodeUpdated_GameThread, STATGROUP_HOLOLENS);
+DECLARE_CYCLE_STAT(TEXT("QR Code Removed"), STAT_FHoloLensARSystem_QRCodeRemoved_GameThread, STATGROUP_HOLOLENS);
 
 static inline FGuid GUIDToFGuid(GUID InGuid)
 {
@@ -58,6 +61,28 @@ struct FMeshUpdateSet
 	TMap<FGuid, FMeshUpdate*> GuidToMeshUpdateList;
 };
 
+/** The UE4 version of the QR code update from the interop */
+struct FQRCodeData
+{
+	FGuid Id;
+
+	FTransform Transform;
+	FVector2D Size;
+	FString QRCode;
+	int32 Version;
+
+	FQRCodeData(QRCodeData* InData)
+		: Id(GUIDToFGuid(InData->Id))
+		, Size(InData->SizeInMeters * 100.f, InData->SizeInMeters * 100.f)
+		, QRCode(InData->Data)
+		, Version(InData->Version)
+	{
+		FVector Translation(InData->Translation[0], InData->Translation[1], InData->Translation[2]);
+		FQuat Orientation(InData->Rotation[0], InData->Rotation[1], InData->Rotation[2], InData->Rotation[3]);
+		Orientation.Normalize();
+		Transform = FTransform(Orientation, Translation);
+	}
+};
 
 FHoloLensARSystem::FHoloLensARSystem()
 	: CameraImage(nullptr)
@@ -725,20 +750,26 @@ void FHoloLensARSystem::ReconcileKnownMeshes(const TArray<FGuid>& KnownMeshes)
 
 void FHoloLensARSystem::QRCodeAdded_Raw(QRCodeData* InCode)
 {
+	FQRCodeData* QRCode = new FQRCodeData(InCode);
 	FHoloLensARSystem* HoloLensARThis = FHoloLensModuleAR::GetHoloLensARSystem().Get();
-	HoloLensARThis->QRCodeAdded(InCode);
+	auto CameraImageTask = FSimpleDelegateGraphTask::FDelegate::CreateThreadSafeSP(HoloLensARThis, &FHoloLensARSystem::QRCodeAdded_GameThread, QRCode);
+	FSimpleDelegateGraphTask::CreateAndDispatchWhenReady(CameraImageTask, GET_STATID(STAT_FHoloLensARSystem_QRCodeAdded_GameThread), nullptr, ENamedThreads::GameThread);
 }
 
 void FHoloLensARSystem::QRCodeUpdated_Raw(QRCodeData* InCode)
 {
+	FQRCodeData* QRCode = new FQRCodeData(InCode);
 	FHoloLensARSystem* HoloLensARThis = FHoloLensModuleAR::GetHoloLensARSystem().Get();
-	HoloLensARThis->QRCodeUpdated(InCode);
+	auto CameraImageTask = FSimpleDelegateGraphTask::FDelegate::CreateThreadSafeSP(HoloLensARThis, &FHoloLensARSystem::QRCodeUpdated_GameThread, QRCode);
+	FSimpleDelegateGraphTask::CreateAndDispatchWhenReady(CameraImageTask, GET_STATID(STAT_FHoloLensARSystem_QRCodeUpdated_GameThread), nullptr, ENamedThreads::GameThread);
 }
 
 void FHoloLensARSystem::QRCodeRemoved_Raw(QRCodeData* InCode)
 {
+	FQRCodeData* QRCode = new FQRCodeData(InCode);
 	FHoloLensARSystem* HoloLensARThis = FHoloLensModuleAR::GetHoloLensARSystem().Get();
-	HoloLensARThis->QRCodeRemoved(InCode);
+	auto CameraImageTask = FSimpleDelegateGraphTask::FDelegate::CreateThreadSafeSP(HoloLensARThis, &FHoloLensARSystem::QRCodeRemoved_GameThread, QRCode);
+	FSimpleDelegateGraphTask::CreateAndDispatchWhenReady(CameraImageTask, GET_STATID(STAT_FHoloLensARSystem_QRCodeRemoved_GameThread), nullptr, ENamedThreads::GameThread);
 }
 
 void FHoloLensARSystem::SetupQRCodeTracking()
@@ -769,91 +800,77 @@ static void DebugDumpQRData(QRCodeData* InCode)
 	UE_LOG(LogHoloLensAR, Log, TEXT("  Orientation = %s"), *Orientation.ToString());
 }
 
-void FHoloLensARSystem::QRCodeAdded(QRCodeData* InCode)
+void FHoloLensARSystem::QRCodeAdded_GameThread(FQRCodeData* InCode)
 {
 	UE_LOG(LogHoloLensAR, Verbose, TEXT("FHoloLensARSystem::QRCodeAdded() called"));
 
 	if (InCode != nullptr)
 	{
-		FGuid Id = GUIDToFGuid(InCode->Id);
 		UARTrackedQRCode* NewQRCode = NewObject<UARTrackedQRCode>();
-		TrackedGeometries.Add(Id, NewQRCode);
-
-		FVector Translation(InCode->Translation[0], InCode->Translation[1], InCode->Translation[2]);
-		FQuat Orientation(InCode->Rotation[0], InCode->Rotation[1], InCode->Rotation[2], InCode->Rotation[3]);
-		Orientation.Normalize();
-		FTransform Transform(Orientation, Translation);
-		FString QRData = InCode->Data;
-		FVector2D Size(InCode->SizeInMeters * 100.f, InCode->SizeInMeters * 100.f);
+		TrackedGeometries.Add(InCode->Id, NewQRCode);
 
 		NewQRCode->UpdateTrackedGeometry(TrackingSystem->GetARCompositionComponent().ToSharedRef(),
 			GFrameCounter,
 			FPlatformTime::Seconds(),
-			Transform,
+			InCode->Transform,
 			// @todo JoeG - add support for an alignment transform
 			FTransform::Identity,
-			Size,
-			QRData,
+			InCode->Size,
+			InCode->QRCode,
 			InCode->Version);
 
 		TriggerOnTrackableAddedDelegates(NewQRCode);
 
 		//		DebugDumpQRData(InCode);
 	}
+	delete InCode;
 }
 
-void FHoloLensARSystem::QRCodeUpdated(QRCodeData* InCode)
+void FHoloLensARSystem::QRCodeUpdated_GameThread(FQRCodeData* InCode)
 {
 	UE_LOG(LogHoloLensAR, Log, TEXT("FHoloLensARSystem::QRCodeUpdated() called"));
 
 	if (InCode != nullptr)
 	{
-		FGuid Id = GUIDToFGuid(InCode->Id);
-		UARTrackedGeometry** FoundGeometry = TrackedGeometries.Find(Id);
+		UARTrackedGeometry** FoundGeometry = TrackedGeometries.Find(InCode->Id);
 		if (FoundGeometry != nullptr)
 		{
 			UARTrackedQRCode* UpdatedQRCode = Cast<UARTrackedQRCode>(*FoundGeometry);
 
-			FVector Translation(InCode->Translation[0], InCode->Translation[1], InCode->Translation[2]);
-			FQuat Orientation(InCode->Rotation[0], InCode->Rotation[1], InCode->Rotation[2], InCode->Rotation[3]);
-			Orientation.Normalize();
-			FTransform Transform(Orientation, Translation);
-			FString QRData = InCode->Data;
-			FVector2D Size(InCode->SizeInMeters * 100.f, InCode->SizeInMeters * 100.f);
-
 			UpdatedQRCode->UpdateTrackedGeometry(TrackingSystem->GetARCompositionComponent().ToSharedRef(),
 				GFrameCounter,
 				FPlatformTime::Seconds(),
-				Transform,
+				InCode->Transform,
 				// @todo JoeG - add support for an alignment transform
 				FTransform::Identity,
-				Size,
-				QRData,
+				InCode->Size,
+				InCode->QRCode,
 				InCode->Version);
 
 			TriggerOnTrackableUpdatedDelegates(UpdatedQRCode);
 		}
 		//		DebugDumpQRData(InCode);
 	}
+	delete InCode;
 }
 
-void FHoloLensARSystem::QRCodeRemoved(QRCodeData* InCode)
+void FHoloLensARSystem::QRCodeRemoved_GameThread(FQRCodeData* InCode)
 {
 	UE_LOG(LogHoloLensAR, Log, TEXT("FHoloLensARSystem::QRCodeRemoved() called"));
 
 	if (InCode != nullptr)
 	{
-		FGuid Id = GUIDToFGuid(InCode->Id);
-		UARTrackedGeometry** FoundGeometry = TrackedGeometries.Find(Id);
+		UARTrackedGeometry** FoundGeometry = TrackedGeometries.Find(InCode->Id);
 		if (FoundGeometry != nullptr)
 		{
 			(*FoundGeometry)->SetTrackingState(EARTrackingState::NotTracking);
 
-			TrackedGeometries.Remove(Id);
+			TrackedGeometries.Remove(InCode->Id);
 			TriggerOnTrackableRemovedDelegates(*FoundGeometry);
 		}
 		//		DebugDumpQRData(InCode);
 	}
+	delete InCode;
 }
 
 #endif
diff --git a/Engine/Plugins/Runtime/AR/Microsoft/HoloLensAR/Source/HoloLensAR/Public/HoloLensARSystem.h b/Engine/Plugins/Runtime/AR/Microsoft/HoloLensAR/Source/HoloLensAR/Public/HoloLensARSystem.h
--- a/Engine/Plugins/Runtime/AR/Microsoft/HoloLensAR/Source/HoloLensAR/Public/HoloLensARSystem.h
+++ b/Engine/Plugins/Runtime/AR/Microsoft/HoloLensAR/Source/HoloLensAR/Public/HoloLensARSystem.h
@@ -27,6 +27,7 @@ struct FMeshUpdate;
 struct FMeshUpdateSet;
 // Forward declaration of the QR code data from the interop layer
 struct QRCodeData;
+struct FQRCodeData;
 
 UCLASS()
 class HOLOLENSAR_API UWMRARPin : public UARPin
@@ -146,9 +147,10 @@ private:
 	static void QRCodeAdded_Raw(QRCodeData* InCode);
 	static void QRCodeUpdated_Raw(QRCodeData* InCode);
 	static void QRCodeRemoved_Raw(QRCodeData* InCode);
-	void QRCodeAdded(QRCodeData* InCode);
-	void QRCodeUpdated(QRCodeData* InCode);
-	void QRCodeRemoved(QRCodeData* InCode);
+	void QRCodeAdded_GameThread(FQRCodeData* InCode);
+	void QRCodeUpdated_GameThread(FQRCodeData* InCode);
+	void QRCodeRemoved_GameThread(FQRCodeData* InCode);
+	//~ QR Code observer callback support
 
 	/**
 	 * Callback from the WinRT layer notifying us of a new camera frame.
