Disable Shader Model 4 shader compilation for HoloLens

diff --git a/Engine/Plugins/Runtime/AR/Microsoft/HoloLensAR/Source/HoloLensTargetPlatform/Private/HoloLensTargetPlatform.cpp b/Engine/Plugins/Runtime/AR/Microsoft/HoloLensAR/Source/HoloLensTargetPlatform/Private/HoloLensTargetPlatform.cpp
--- a/Engine/Plugins/Runtime/AR/Microsoft/HoloLensAR/Source/HoloLensTargetPlatform/Private/HoloLensTargetPlatform.cpp
+++ b/Engine/Plugins/Runtime/AR/Microsoft/HoloLensAR/Source/HoloLensTargetPlatform/Private/HoloLensTargetPlatform.cpp
@@ -121,21 +121,18 @@ void FHoloLensTargetPlatform::GetAllTextureFormats(TArray<FName>& OutFormats) co
 }
 
 static FName NAME_PCD3D_ES3_1(TEXT("PCD3D_ES31"));
-static FName NAME_PCD3D_SM4(TEXT("PCD3D_SM4"));
 static FName NAME_PCD3D_SM5(TEXT("PCD3D_SM5"));
 
 void FHoloLensTargetPlatform::GetAllPossibleShaderFormats(TArray<FName>& OutFormats) const
 {
 	OutFormats.AddUnique(NAME_PCD3D_ES3_1);
 	OutFormats.AddUnique(NAME_PCD3D_SM5);
-	OutFormats.AddUnique(NAME_PCD3D_SM4);
 }
 
 void FHoloLensTargetPlatform::GetAllTargetedShaderFormats(TArray<FName>& OutFormats) const
 {
 	OutFormats.AddUnique(NAME_PCD3D_ES3_1);
 	OutFormats.AddUnique(NAME_PCD3D_SM5);
-	OutFormats.AddUnique(NAME_PCD3D_SM4);
 }
 
 #endif
