UE-82961 Hololens 2 anchors have identity position immediately after load.
-Update transform and tracking status immediately on load from anchor store.  This will either get the transform or set the tracking status to NotTracked.  Then one can either use the transform, or know that it is not yet available.

diff --git a/Engine/Plugins/Runtime/AR/Microsoft/HoloLensAR/Source/HoloLensAR/Private/HoloLensARSystemAnchors.cpp b/Engine/Plugins/Runtime/AR/Microsoft/HoloLensAR/Source/HoloLensAR/Private/HoloLensARSystemAnchors.cpp
--- a/Engine/Plugins/Runtime/AR/Microsoft/HoloLensAR/Source/HoloLensAR/Private/HoloLensARSystemAnchors.cpp
+++ b/Engine/Plugins/Runtime/AR/Microsoft/HoloLensAR/Source/HoloLensAR/Private/HoloLensARSystemAnchors.cpp
@@ -95,8 +95,11 @@ TArray<UWMRARPin*> FHoloLensARSystem::LoadWMRAnchorStoreARPins()
 	TSharedPtr<FARSupportInterface, ESPMode::ThreadSafe> ARSupportInterface = TrackingSystem->GetARCompositionComponent();
 	for (FString& AnchorId : AnchorIds)
 	{
+		FTransform Transform;
+		const bool bTracked = WMRGetAnchorTransform(*AnchorId, Transform);
+
 		UWMRARPin* NewPin = NewObject<UWMRARPin>();
-		NewPin->InitARPin(ARSupportInterface.ToSharedRef(), nullptr, FTransform::Identity, nullptr, FName(*AnchorId));
+		NewPin->InitARPin(ARSupportInterface.ToSharedRef(), nullptr, Transform, nullptr, FName(*AnchorId));
 
 		AnchorIdToPinMap.Add(AnchorId, NewPin);
 		NewPin->SetAnchorId(AnchorId);
@@ -104,6 +107,8 @@ TArray<UWMRARPin*> FHoloLensARSystem::LoadWMRAnchorStoreARPins()
 
 		Pins.Add(NewPin);
 		LoadedPins.Add(NewPin);
+
+		NewPin->OnTrackingStateChanged(bTracked ? EARTrackingState::Tracking : EARTrackingState::NotTracking);
 	}
 
 	return LoadedPins;
