Fix static analysis warning C6387: 'PackageFullName' could be '0':  this does not adhere to the specification for the function (for 2 functions in HololensTargetDevice.cpp).

diff --git a/Engine/Plugins/Runtime/AR/Microsoft/HoloLensAR/Source/HoloLensTargetPlatform/Private/HoloLensTargetDevice.cpp b/Engine/Plugins/Runtime/AR/Microsoft/HoloLensAR/Source/HoloLensTargetPlatform/Private/HoloLensTargetDevice.cpp
--- a/Engine/Plugins/Runtime/AR/Microsoft/HoloLensAR/Source/HoloLensTargetPlatform/Private/HoloLensTargetDevice.cpp
+++ b/Engine/Plugins/Runtime/AR/Microsoft/HoloLensAR/Source/HoloLensTargetPlatform/Private/HoloLensTargetDevice.cpp
@@ -113,7 +113,8 @@ namespace WindowsMixedReality
 			if (SUCCEEDED(ManifestReader->GetPackageId(&ManifestPackageId)))
 			{
 				LPWSTR PackageFullName = nullptr;
-				if (SUCCEEDED(ManifestPackageId->GetPackageFullName(&PackageFullName)))
+				// Avoid warning C6387 (static analysis assumes PackageFullName could be null even if the GetPackageFullName call succeeds.
+				if (SUCCEEDED(ManifestPackageId->GetPackageFullName(&PackageFullName)) && PackageFullName != nullptr)
 				{
 					PackageDebugSettings->EnableDebugging(PackageFullName, nullptr, nullptr);
 					CoTaskMemFree(PackageFullName);
@@ -122,7 +123,8 @@ namespace WindowsMixedReality
 		}
 
 		LPWSTR Aumid = nullptr;
-		if (FAILED(ApplicationMetadata->GetAppUserModelId(&Aumid)))
+		// Avoid warning C6387 (static analysis assumes Aumid could be null even if the GetAppUserModelId call succeeds.
+		if (FAILED(ApplicationMetadata->GetAppUserModelId(&Aumid)) || Aumid == nullptr)
 		{
 			return false;
 		}
