Hololens static analysis warning about unused return value

diff --git a/Engine/Plugins/Runtime/AR/Microsoft/HoloLensAR/Source/HoloLensTargetPlatform/Private/HoloLensDeviceDetector.cpp b/Engine/Plugins/Runtime/AR/Microsoft/HoloLensAR/Source/HoloLensTargetPlatform/Private/HoloLensDeviceDetector.cpp
--- a/Engine/Plugins/Runtime/AR/Microsoft/HoloLensAR/Source/HoloLensTargetPlatform/Private/HoloLensDeviceDetector.cpp
+++ b/Engine/Plugins/Runtime/AR/Microsoft/HoloLensAR/Source/HoloLensTargetPlatform/Private/HoloLensDeviceDetector.cpp
@@ -163,7 +163,8 @@ void FHoloLensDeviceDetector::StartDeviceDetection()
 
 	HRESULT hr;
 
-	::RoInitialize(RO_INIT_MULTITHREADED);
+	hr = ::RoInitialize(RO_INIT_MULTITHREADED);
+	if (FAILED(hr)) { return; }
 
 	ComPtr <IDeviceInformationStatics2> DeviceInformationStatics2;
 	hr = GetActivationFactory(HStringReference(RuntimeClass_Windows_Devices_Enumeration_DeviceInformation).Get(), &DeviceInformationStatics2);
