Fix static analysis warning (potential dereference null pointer)

diff --git a/Engine/Plugins/Runtime/AR/Microsoft/HoloLensAR/Source/HoloLensTargetPlatform/Private/HoloLensDeviceDetector.cpp b/Engine/Plugins/Runtime/AR/Microsoft/HoloLensAR/Source/HoloLensTargetPlatform/Private/HoloLensDeviceDetector.cpp
--- a/Engine/Plugins/Runtime/AR/Microsoft/HoloLensAR/Source/HoloLensTargetPlatform/Private/HoloLensDeviceDetector.cpp
+++ b/Engine/Plugins/Runtime/AR/Microsoft/HoloLensAR/Source/HoloLensTargetPlatform/Private/HoloLensDeviceDetector.cpp
@@ -344,7 +344,7 @@ void FHoloLensDeviceDetector::OnDeviceWatcherDeviceAdded(IDeviceInformation* Inf
 		if (FAILED(hr)) { return; }
 
 		hr = PropertyValue->GetStringArray(&uiValue, &pArrayString);
-		if (FAILED(hr)) { return; }
+		if (FAILED(hr) || pArrayString == nullptr) { return; }
 
 		if (uiValue == 0) { return; }
 
