fixes a crash when the editor is run as a dedicated server. Submitted by Microsoft.

diff --git a/Engine/Plugins/Runtime/AR/Microsoft/HoloLensAR/Source/HoloLensAR/Private/HoloLensARSystem.cpp b/Engine/Plugins/Runtime/AR/Microsoft/HoloLensAR/Source/HoloLensAR/Private/HoloLensARSystem.cpp
--- a/Engine/Plugins/Runtime/AR/Microsoft/HoloLensAR/Source/HoloLensAR/Private/HoloLensARSystem.cpp
+++ b/Engine/Plugins/Runtime/AR/Microsoft/HoloLensAR/Source/HoloLensAR/Private/HoloLensARSystem.cpp
@@ -252,7 +252,7 @@ void FHoloLensARSystem::OnStartARSession(UARSessionConfig* InSessionConfig)
 #if WITH_EDITOR
 	UEditorEngine* EditorEngine = CastChecked<UEditorEngine>(GEngine);
 	FSceneViewport* PIEViewport = (FSceneViewport*)EditorEngine->GetPIEViewport();
-	if (!PIEViewport->IsStereoRenderingAllowed())
+	if (PIEViewport != nullptr && !PIEViewport->IsStereoRenderingAllowed())
 	{
 		// Running the AR session on a non-stereo window will break spatial anchors.
 		SessionStatus.Status = EARSessionStatus::NotSupported;
