UE-91210 Hololens - restore wmr/hololensar plugin dependency direction to previous
-Removed unnecessary build.cs include path to WMR which caused a dependency warning.

diff --git a/Engine/Plugins/Runtime/AR/Microsoft/HoloLensAR/Source/HoloLensAR/HoloLensAR.Build.cs b/Engine/Plugins/Runtime/AR/Microsoft/HoloLensAR/Source/HoloLensAR/HoloLensAR.Build.cs
--- a/Engine/Plugins/Runtime/AR/Microsoft/HoloLensAR/Source/HoloLensAR/HoloLensAR.Build.cs
+++ b/Engine/Plugins/Runtime/AR/Microsoft/HoloLensAR/Source/HoloLensAR/HoloLensAR.Build.cs
@@ -36,7 +36,6 @@ public class HoloLensAR : ModuleRules
 		{
 			PrivateDependencyModuleNames.Add("UnrealEd");
 		}
-        PrivateIncludePathModuleNames.Add("WindowsMixedRealityHMD");
         if (Target.Platform == UnrealTargetPlatform.Win64)
         {
             RuntimeDependencies.Add(System.IO.Path.Combine("$(EngineDir)/Binaries/ThirdParty/Windows/x64", "Microsoft.MixedReality.QR.dll"));
