UE-91380 Hololens - latency, remoting, hand tracking blueprint, and logging fixes
-Latency improvement - block for next frame before updating render thread frame, not after.
-HL2 Remoting fix - allow much more max rhi frame latency because multiple windows (including the VR preview window) count toward the total and we don't want this code to limit framerate, specific HMD code does that.
-Marked some hand tracking related enums and structs as blueprinttype so they can be used in blueprint.
-Fixed some wrongly worded logging for spatial anchors.

diff --git a/Engine/Plugins/Runtime/AR/Microsoft/HoloLensAR/Source/HoloLensAR/Private/HoloLensARSystemAnchors.cpp b/Engine/Plugins/Runtime/AR/Microsoft/HoloLensAR/Source/HoloLensAR/Private/HoloLensARSystemAnchors.cpp
--- a/Engine/Plugins/Runtime/AR/Microsoft/HoloLensAR/Source/HoloLensAR/Private/HoloLensARSystemAnchors.cpp
+++ b/Engine/Plugins/Runtime/AR/Microsoft/HoloLensAR/Source/HoloLensAR/Private/HoloLensARSystemAnchors.cpp
@@ -21,7 +21,7 @@ UWMRARPin* FHoloLensARSystem::CreateNamedARPin(FName Name, const FTransform& Pin
 
 	if (AnchorIdToPinMap.Contains(WMRAnchorId))
 	{
-		UE_LOG(LogHoloLensAR, Warning, TEXT("CreateWMRAnchorStoreARPin: Creation of Anchor %s failed because that anchorID is already in use!  No pin created."), *WMRAnchorId);
+		UE_LOG(LogHoloLensAR, Warning, TEXT("CreateNamedARPin: Creation of Anchor %s failed because that anchorID is already in use!  No pin created."), *WMRAnchorId);
 		return nullptr;
 	}
 
@@ -29,7 +29,7 @@ UWMRARPin* FHoloLensARSystem::CreateNamedARPin(FName Name, const FTransform& Pin
 		bool bSuccess = WMRCreateAnchor(*WMRAnchorId, PinToTrackingTransform.GetLocation(), PinToTrackingTransform.GetRotation());
 		if (!bSuccess)
 		{
-			UE_LOG(LogHoloLensAR, Warning, TEXT("CreateWMRAnchorStoreARPin: Creation of Anchor %s failed!  No anchor or pin created."), *WMRAnchorId);
+			UE_LOG(LogHoloLensAR, Warning, TEXT("CreateNamedARPin: Creation of Anchor %s failed!  No anchor or pin created."), *WMRAnchorId);
 			return nullptr;
 		}
 	}
@@ -69,12 +69,12 @@ bool FHoloLensARSystem::PinComponentToARPin(USceneComponent* ComponentToPin, UWM
 {
 	if (Pin == nullptr)
 	{
-		UE_LOG(LogHoloLensAR, Warning, TEXT("PinComponentToWMRAnchorStoreARPin: Pin was null.  Doing nothing."));
+		UE_LOG(LogHoloLensAR, Warning, TEXT("PinComponentToARPin: Pin was null.  Doing nothing."));
 		return false;
 	}
 	if (ComponentToPin == nullptr)
 	{
-		UE_LOG(LogHoloLensAR, Warning, TEXT("PinComponentToWMRAnchorStoreARPin: Tried to pin null component to pin %s.  Doing nothing."), *Pin->GetDebugName().ToString());
+		UE_LOG(LogHoloLensAR, Warning, TEXT("PinComponentToARPin: Tried to pin null component to pin %s.  Doing nothing."), *Pin->GetDebugName().ToString());
 		return false;
 	}
 
@@ -83,12 +83,12 @@ bool FHoloLensARSystem::PinComponentToARPin(USceneComponent* ComponentToPin, UWM
 		{
 			if (FindResult == Pin)
 			{
-				UE_LOG(LogHoloLensAR, Warning, TEXT("PinComponentToWMRAnchorStoreARPin: Component %s is already pinned to pin %s.  Doing nothing."), *ComponentToPin->GetReadableName(), *Pin->GetDebugName().ToString());
+				UE_LOG(LogHoloLensAR, Warning, TEXT("PinComponentToARPin: Component %s is already pinned to pin %s.  Doing nothing."), *ComponentToPin->GetReadableName(), *Pin->GetDebugName().ToString());
 				return true;
 			}
 			else
 			{
-				UE_LOG(LogHoloLensAR, Warning, TEXT("PinComponentToWMRAnchorStoreARPin: Component %s is pinned to pin %s. Unpinning it from that pin first.  The pin will not be destroyed."), *ComponentToPin->GetReadableName(), *Pin->GetDebugName().ToString());
+				UE_LOG(LogHoloLensAR, Warning, TEXT("PinComponentToARPin: Component %s is pinned to pin %s. Unpinning it from that pin first.  The pin will not be destroyed."), *ComponentToPin->GetReadableName(), *Pin->GetDebugName().ToString());
 				FindResult->SetPinnedComponent(nullptr);
 			}
 		}
diff --git a/Engine/Plugins/Runtime/WindowsMixedReality/Source/WindowsMixedRealityHMD/Public/WindowsMixedRealityFunctionLibrary.h b/Engine/Plugins/Runtime/WindowsMixedReality/Source/WindowsMixedRealityHMD/Public/WindowsMixedRealityFunctionLibrary.h
--- a/Engine/Plugins/Runtime/WindowsMixedReality/Source/WindowsMixedRealityHMD/Public/WindowsMixedRealityFunctionLibrary.h
+++ b/Engine/Plugins/Runtime/WindowsMixedReality/Source/WindowsMixedRealityHMD/Public/WindowsMixedRealityFunctionLibrary.h
@@ -8,7 +8,7 @@
 
 #include "WindowsMixedRealityFunctionLibrary.Generated.h"
 
-UENUM()
+UENUM(BlueprintType, Category = "WindowsMixedRealityHMD")
 enum class EHMDSpatialLocatability : uint8 
 {
 	Unavailable = 0,
@@ -18,7 +18,7 @@ enum class EHMDSpatialLocatability : uint8
 	PositionalTrackingInhibited = 4,
 };
 
-UENUM()
+UENUM(BlueprintType, Category = "WindowsMixedRealityHMD")
 enum class EHMDInputControllerButtons : uint8
 {
 	Select,
@@ -30,7 +30,7 @@ enum class EHMDInputControllerButtons : uint8
 	Count UMETA(Hidden)
 };
 
-UENUM()
+UENUM(BlueprintType, Category = "WindowsMixedRealityHMD")
 enum class EHMDTrackingStatus : uint8
 {
 	NotTracked,
@@ -38,7 +38,7 @@ enum class EHMDTrackingStatus : uint8
 	Tracked
 };
 
-USTRUCT(BlueprintType)
+USTRUCT(BlueprintType, Category = "WindowsMixedRealityHMD")
 struct FPointerPoseInfo
 {
 	GENERATED_USTRUCT_BODY()
diff --git a/Engine/Plugins/Runtime/WindowsMixedReality/Source/WindowsMixedRealityRHI/Private/WindowsMixedRealityRHI.cpp b/Engine/Plugins/Runtime/WindowsMixedReality/Source/WindowsMixedRealityRHI/Private/WindowsMixedRealityRHI.cpp
--- a/Engine/Plugins/Runtime/WindowsMixedReality/Source/WindowsMixedRealityRHI/Private/WindowsMixedRealityRHI.cpp
+++ b/Engine/Plugins/Runtime/WindowsMixedReality/Source/WindowsMixedRealityRHI/Private/WindowsMixedRealityRHI.cpp
@@ -173,6 +173,14 @@ FViewportRHIRef FWindowsMixedRealityDynamicRHI::RHICreateViewport(void* WindowHa
 		PreferredPixelFormat = EDefaultBackBufferPixelFormat::Convert2PixelFormat(EDefaultBackBufferPixelFormat::FromInt(CVarDefaultBackBufferPixelFormat->GetValueOnGameThread()));
 	}
 
+	if (GIsEditor)
+	{
+		// DXGI device maximum frame latency counts the total number of presents enqueued rather than logical engine frames in flight. Inside the editor we can easily
+		// run into the default limit of 3 presents for different HWNDs when a VR preview window is active, causing frame rate to be throttled by the local monitor.
+		// Increase the limit so that the HMD can properly take ownership of frame pacing instead.
+		IConsoleManager::Get().FindConsoleVariable(TEXT("RHI.MaximumFrameLatency"))->Set(16);
+	}
+
 	return new FWindowsMixedRealityViewport(this, (HWND)WindowHandle, SizeX, SizeY, bIsFullscreen, PreferredPixelFormat);
 }
 
