Fix a typo

diff --git a/Engine/Plugins/Runtime/AR/Microsoft/HoloLensAR/Source/HoloLensTargetPlatform/Private/HoloLensTargetPlatform.h b/Engine/Plugins/Runtime/AR/Microsoft/HoloLensAR/Source/HoloLensTargetPlatform/Private/HoloLensTargetPlatform.h
--- a/Engine/Plugins/Runtime/AR/Microsoft/HoloLensAR/Source/HoloLensTargetPlatform/Private/HoloLensTargetPlatform.h
+++ b/Engine/Plugins/Runtime/AR/Microsoft/HoloLensAR/Source/HoloLensTargetPlatform/Private/HoloLensTargetPlatform.h
@@ -71,7 +71,7 @@ public:
 
 	virtual const class FStaticMeshLODSettings& GetStaticMeshLODSettings() const override { return StaticMeshLODSettings; }
 
-	virtual FSkeletalMeshDefaultLODStreamingSettings& GetSkeletalMeshDefaultLODStreamingSettings() const override { return SkeletalMeshDefaultLODStreamingSettings; }
+	virtual const FSkeletalMeshDefaultLODStreamingSettings& GetSkeletalMeshDefaultLODStreamingSettings() const override { return SkeletalMeshDefaultLODStreamingSettings; }
 
 	virtual void GetTextureFormats(const UTexture* InTexture, TArray< TArray<FName> >& OutFormats) const override;
 
