Fixed HoloLensPlatformEditor linker error.

diff --git a/Engine/Plugins/Runtime/AR/Microsoft/HoloLensAR/Source/HoloLensPlatformEditor/HoloLensPlatformEditor.Build.cs b/Engine/Plugins/Runtime/AR/Microsoft/HoloLensAR/Source/HoloLensPlatformEditor/HoloLensPlatformEditor.Build.cs
--- a/Engine/Plugins/Runtime/AR/Microsoft/HoloLensAR/Source/HoloLensPlatformEditor/HoloLensPlatformEditor.Build.cs
+++ b/Engine/Plugins/Runtime/AR/Microsoft/HoloLensAR/Source/HoloLensPlatformEditor/HoloLensPlatformEditor.Build.cs
@@ -11,7 +11,7 @@ public class HoloLensPlatformEditor : ModuleRules
 				"Core",
 				"CoreUObject",
 				"Settings",
-				"TargetPlatform",
+				"DeveloperToolSettings",
 				"DesktopPlatform",
 				"Slate",
 				"SlateCore",
