Possible fix to nullptr dereference in HL Target Device

diff --git a/Engine/Plugins/Runtime/AR/Microsoft/HoloLensAR/Source/HoloLensTargetPlatform/Private/HoloLensTargetDevice.cpp b/Engine/Plugins/Runtime/AR/Microsoft/HoloLensAR/Source/HoloLensTargetPlatform/Private/HoloLensTargetDevice.cpp
--- a/Engine/Plugins/Runtime/AR/Microsoft/HoloLensAR/Source/HoloLensTargetPlatform/Private/HoloLensTargetDevice.cpp
+++ b/Engine/Plugins/Runtime/AR/Microsoft/HoloLensAR/Source/HoloLensTargetPlatform/Private/HoloLensTargetDevice.cpp
@@ -172,7 +172,7 @@ FHoloLensTargetDevice::~FHoloLensTargetDevice()
 	if (!Info.IsLocal)
 	{
 #if WITH_ENGINE
-		if (TimerHandle.IsValid())
+		if ((GWorld != nullptr) && TimerHandle.IsValid())
 		{
 			GWorld->GetTimerManager().ClearTimer(TimerHandle);
 		}
