Changed some ITargetPlatform methods (GetAllSupportedShaderFormats, GetTargetedShaderFormats) to remain compiled with a lack of WITH_ENGINE. Their omission isn't necessary from a dependency standpoint, and their access will be required in a future change.

diff --git a/Engine/Plugins/Runtime/AR/Microsoft/HoloLensAR/Source/HoloLensTargetPlatform/Private/HoloLensTargetPlatform.h b/Engine/Plugins/Runtime/AR/Microsoft/HoloLensAR/Source/HoloLensTargetPlatform/Private/HoloLensTargetPlatform.h
--- a/Engine/Plugins/Runtime/AR/Microsoft/HoloLensAR/Source/HoloLensTargetPlatform/Private/HoloLensTargetPlatform.h
+++ b/Engine/Plugins/Runtime/AR/Microsoft/HoloLensAR/Source/HoloLensTargetPlatform/Private/HoloLensTargetPlatform.h
@@ -63,6 +63,10 @@ public:
 
 	virtual bool SupportsBuildTarget(EBuildTargetType BuildTarget) const override;
 
+	virtual void GetAllPossibleShaderFormats(TArray<FName>& OutFormats) const override;
+
+	virtual void GetAllTargetedShaderFormats(TArray<FName>& OutFormats) const override;	
+
 #if WITH_ENGINE
 	virtual void GetReflectionCaptureFormats(TArray<FName>& OutFormats) const override;
 
@@ -79,12 +83,6 @@ public:
 		TextureLODSettings = InTextureLODSettings;
 	}
 
-	virtual void GetAllPossibleShaderFormats(TArray<FName>& OutFormats) const override;
-
-	virtual void GetAllTargetedShaderFormats(TArray<FName>& OutFormats) const override;
-	
-	//virtual void GetAllCachedShaderFormats( TArray<FName>& OutFormats ) const override {}
-
 	virtual FName GetWaveFormat( const class USoundWave* Wave ) const override
 	{
 		static FName NAME_OGG(TEXT("OGG"));
diff --git a/Engine/Source/Developer/TargetPlatform/Public/Common/TargetPlatformBase.h b/Engine/Source/Developer/TargetPlatform/Public/Common/TargetPlatformBase.h
--- a/Engine/Source/Developer/TargetPlatform/Public/Common/TargetPlatformBase.h
+++ b/Engine/Source/Developer/TargetPlatform/Public/Common/TargetPlatformBase.h
@@ -137,7 +137,6 @@ public:
 
 	virtual bool SupportsValueForType(FName SupportedType, FName RequiredSupportedValue) const override
 	{
-#if WITH_ENGINE
 		// check if the given shader format is returned by this TargetPlatform
 		if (SupportedType == TEXT("ShaderFormat"))
 		{
@@ -145,7 +144,7 @@ public:
 			GetAllPossibleShaderFormats(AllPossibleShaderFormats);
 			return AllPossibleShaderFormats.Contains(RequiredSupportedValue);
 		}
-#endif
+
 		return false;
 	}
 
